Copying Orion-Staging @ 3748653 (Orion/Dev-General @ 3722124) to //UE4/Main

diff --git a/Engine/Plugins/Runtime/ImmediatePhysics/Source/ImmediatePhysics/Private/BoneControllers/AnimNode_RigidBody.cpp b/Engine/Plugins/Runtime/ImmediatePhysics/Source/ImmediatePhysics/Private/BoneControllers/AnimNode_RigidBody.cpp
--- a/Engine/Plugins/Runtime/ImmediatePhysics/Source/ImmediatePhysics/Private/BoneControllers/AnimNode_RigidBody.cpp
+++ b/Engine/Plugins/Runtime/ImmediatePhysics/Source/ImmediatePhysics/Private/BoneControllers/AnimNode_RigidBody.cpp
@@ -6,6 +6,7 @@
 #include "PhysicsEngine/BodySetup.h"
 #include "PhysicsEngine/PhysicsAsset.h"
 #include "PhysicsEngine/PhysicsConstraintTemplate.h"
+#include "GameFramework/PawnMovementComponent.h"
 
 #include "ImmediatePhysicsSimulation.h"
 #include "ImmediatePhysicsActorHandle.h"
@@ -36,6 +37,8 @@ FAnimNode_RigidBody::FAnimNode_RigidBody():
 	TotalMass = 0.f;
 	CachedBounds.W = 0;
 	UnsafeWorld = nullptr;
+	bSimulationStarted = false;
+	bCheckForBodyTransformInit = false;
 }
 
 FAnimNode_RigidBody::~FAnimNode_RigidBody()
@@ -53,7 +56,11 @@ void FAnimNode_RigidBody::GatherDebugData(FNodeDebugData& DebugData)
 
 	DebugData.AddDebugItem(DebugLine);
 
-	ComponentPose.GatherDebugData(DebugData);
+	const bool bUsingFrozenPose = bFreezeIncomingPoseOnStart && bSimulationStarted && (CapturedFrozenPose.GetPose().GetNumBones() > 0);
+	if (!bUsingFrozenPose)
+	{
+		ComponentPose.GatherDebugData(DebugData);
+	}
 }
 
 FVector WorldVectorToSpaceNoScale(ESimulationSpace Space, const FVector& WorldDir, const FTransform& ComponentToWorld, const FTransform& RootBoneTM)
@@ -78,86 +85,267 @@ FVector WorldPositionToSpace(ESimulationSpace Space, const FVector& WorldPoint,
 	}
 }
 
+FORCEINLINE_DEBUGGABLE FTransform ConvertCSTransformToSimSpace(ESimulationSpace SimulationSpace, const FTransform& InCSTransform, const FTransform& ComponentToWorld, const FTransform& RootBoneTM)
+{
+	switch (SimulationSpace)
+	{
+		case ESimulationSpace::ComponentSpace: return InCSTransform;
+		case ESimulationSpace::WorldSpace:  return InCSTransform * ComponentToWorld; 
+		case ESimulationSpace::RootBoneSpace: return InCSTransform.GetRelativeTransform(RootBoneTM); break;
+		default: ensureMsgf(false, TEXT("Unsupported Simulation Space")); return InCSTransform;
+	}
+}
 
+void FAnimNode_RigidBody::UpdateComponentPose_AnyThread(const FAnimationUpdateContext& Context)
+{
+	// Only freeze update graph after initial update, as we want to get that pose through.
+	if (bFreezeIncomingPoseOnStart && bSimulationStarted && !bResetSimulated)
+	{
+		// If we have a Frozen Pose captured, 
+		// then we don't need to update the rest of the graph.
+		if (CapturedFrozenPose.GetPose().GetNumBones() > 0)
+		{
+		}
+		else
+		{
+			// Create a new context with zero deltatime to freeze time in rest of the graph.
+			// This will be used to capture a frozen pose.
+			FAnimationUpdateContext FrozenContext = Context;
+			FrozenContext.FractionalWeightAndTime(1.f, 0.f);
+
+			Super::UpdateComponentPose_AnyThread(FrozenContext);
+		}
+	}
+	else
+	{
+		Super::UpdateComponentPose_AnyThread(Context);
+	}
+}
+
+void FAnimNode_RigidBody::EvaluateComponentPose_AnyThread(FComponentSpacePoseContext& Output)
+{
+	if (bFreezeIncomingPoseOnStart && bSimulationStarted)
+	{
+		// If we have a Frozen Pose captured, use it.
+		// Only after our intialize setup. As we need new pose for that.
+		if (!bResetSimulated && (CapturedFrozenPose.GetPose().GetNumBones() > 0))
+		{
+			Output.Pose.CopyPose(CapturedFrozenPose);
+			Output.Curve.CopyFrom(CapturedFrozenCurves);
+		}
+		// Otherwise eval graph to capture it.
+		else
+		{
+			Super::EvaluateComponentPose_AnyThread(Output);
+			CapturedFrozenPose.CopyPose(Output.Pose);
+			CapturedFrozenCurves.CopyFrom(Output.Curve);
+		}
+	}
+	else
+	{
+		Super::EvaluateComponentPose_AnyThread(Output);
+	}
+
+	// Capture incoming pose if 'bTransferBoneVelocities' is set.
+	// That is, until simulation starts.
+	if (bTransferBoneVelocities && !bSimulationStarted)
+	{
+		CapturedBoneVelocityPose.CopyPose(Output.Pose);
+		CapturedBoneVelocityPose.CopyAndAssignBoneContainer(CapturedBoneVelocityBoneContainer);
+	}
+}
+
+void FAnimNode_RigidBody::InitializeNewBodyTransformsDuringSimulation(FComponentSpacePoseContext& Output, const FTransform& ComponentTransform, const FTransform& RootBoneTM)
+{
+	for (const FOutputBoneData& OutputData : OutputBoneData)
+	{
+		const int32 BodyIndex = OutputData.BodyIndex;
+		FBodyAnimData& BodyData = BodyAnimData[BodyIndex];
+		if (!BodyData.bBodyTransformInitialized)
+		{
+			BodyData.bBodyTransformInitialized = true;
+
+			// If we have a parent body, we need to grab relative transforms to it.
+			if (OutputData.ParentBodyIndex != INDEX_NONE)
+			{
+				ensure(BodyAnimData[OutputData.ParentBodyIndex].bBodyTransformInitialized);
+
+				FTransform BodyRelativeTransform = FTransform::Identity;
+				for (const FCompactPoseBoneIndex CompactBoneIndex : OutputData.BoneIndicesToParentBody)
+				{
+					const FTransform& LocalSpaceTM = Output.Pose.GetLocalSpaceTransform(CompactBoneIndex);
+					BodyRelativeTransform = BodyRelativeTransform * LocalSpaceTM;
+				}
+
+				const FTransform WSBodyTM = BodyRelativeTransform * Bodies[OutputData.ParentBodyIndex]->GetWorldTransform();
+				Bodies[BodyIndex]->SetWorldTransform(WSBodyTM);
+			}
+			// If we don't have a parent body, then we can just grab the incoming pose in component space.
+			else
+			{
+				const FTransform& ComponentSpaceTM = Output.Pose.GetComponentSpaceTransform(OutputData.CompactPoseBoneIndex);
+				const FTransform BodyTM = ConvertCSTransformToSimSpace(SimulationSpace, ComponentSpaceTM, ComponentTransform, RootBoneTM);
+
+				Bodies[BodyIndex]->SetWorldTransform(BodyTM);
+			}
+		}
+	}
+}
+
+DECLARE_CYCLE_STAT(TEXT("RigidBody_Eval"), STAT_RigidBody_Eval, STATGROUP_Anim);
 DECLARE_CYCLE_STAT(TEXT("FAnimNode_Ragdoll::EvaluateSkeletalControl_AnyThread"), STAT_ImmediateEvaluateSkeletalControl, STATGROUP_ImmediatePhysics);
 
 void FAnimNode_RigidBody::EvaluateSkeletalControl_AnyThread(FComponentSpacePoseContext& Output, TArray<FBoneTransform>& OutBoneTransforms)
 {
+	SCOPE_CYCLE_COUNTER(STAT_RigidBody_Eval);
 	SCOPE_CYCLE_COUNTER(STAT_ImmediateEvaluateSkeletalControl);
 	//FPlatformMisc::BeginNamedEvent(FColor::Magenta, "FAnimNode_Ragdoll::EvaluateSkeletalControl_AnyThread");
 
-	if(PhysicsSimulation && DeltaSeconds > 0)
+	const float DeltaSeconds = AccumulatedDeltaTime;
+	AccumulatedDeltaTime = 0.f;
+
+	if (PhysicsSimulation)
 	{
 		const FBoneContainer& BoneContainer = Output.Pose.GetPose().GetBoneContainer();
-
 		const FTransform CompWorldSpaceTM = Output.AnimInstanceProxy->GetComponentTransform();
 		const FTransform RootBoneTM = Output.Pose.GetComponentSpaceTransform(RootBoneRef.GetCompactPoseIndex(BoneContainer));
-		//Update physics with transforms
-		
-		for (const FOutputBoneData& OutputData : OutputBoneData)
+
+		// Initialize potential new bodies because of LOD change.
+		if (!bResetSimulated && bCheckForBodyTransformInit)
 		{
-			int32 BodyIndex = OutputData.BodyIndex;
-			if(BodyIndex != INDEX_NONE && (bResetSimulated || !IsSimulated[BodyIndex]))
+			bCheckForBodyTransformInit = false;
+			InitializeNewBodyTransformsDuringSimulation(Output, CompWorldSpaceTM, RootBoneTM);
+		}
+
+		// If time advances, update simulation
+		if (DeltaSeconds > 0)
+		{
+			// First update, initialize bodies from animation.
+			if (bResetSimulated)
 			{
-				FCompactPoseBoneIndex SimBoneIndex = OutputData.BoneReference.GetCompactPoseIndex(BoneContainer);
-				const FTransform& ComponentSpaceTM = Output.Pose.GetComponentSpaceTransform(SimBoneIndex);
-				
-				FTransform BodyTM;
-				switch(SimulationSpace)
+				bResetSimulated = false;
+
+				// Capture bone velocities if we have captured a bone velocity pose.
+				if (bTransferBoneVelocities && (CapturedBoneVelocityPose.GetPose().GetNumBones() > 0))
 				{
-					case ESimulationSpace::ComponentSpace : BodyTM = ComponentSpaceTM; break;
-					case ESimulationSpace::WorldSpace : BodyTM = ComponentSpaceTM * CompWorldSpaceTM; break;
-					case ESimulationSpace::RootBoneSpace: BodyTM = ComponentSpaceTM.GetRelativeTransform(RootBoneTM); break;
-					default: ensureMsgf(false, TEXT("Unsupported Simulation Space")); BodyTM = ComponentSpaceTM;
+					for (const FOutputBoneData& OutputData : OutputBoneData)
+					{
+						const int32 BodyIndex = OutputData.BodyIndex;
+						FBodyAnimData& BodyData = BodyAnimData[BodyIndex];
+
+						if (BodyData.bIsSimulated)
+						{
+							const FCompactPoseBoneIndex NextCompactPoseBoneIndex = OutputData.CompactPoseBoneIndex;
+							// Convert CompactPoseBoneIndex to SkeletonBoneIndex...
+							const int32 PoseSkeletonBoneIndex = BoneContainer.GetPoseToSkeletonBoneIndexArray()[NextCompactPoseBoneIndex.GetInt()];
+							// ... So we can convert to the captured pose CompactPoseBoneIndex. 
+							// In case there was a LOD change, and poses are not compatible anymore.
+							const FCompactPoseBoneIndex PrevCompactPoseBoneIndex = CapturedBoneVelocityBoneContainer.GetCompactPoseIndexFromSkeletonIndex(PoseSkeletonBoneIndex);
+
+							if (PrevCompactPoseBoneIndex != FCompactPoseBoneIndex(INDEX_NONE))
+							{
+								const FTransform PrevCSTM = CapturedBoneVelocityPose.GetComponentSpaceTransform(PrevCompactPoseBoneIndex);
+								const FTransform NextCSTM = Output.Pose.GetComponentSpaceTransform(NextCompactPoseBoneIndex);
+
+								const FTransform PrevSSTM = ConvertCSTransformToSimSpace(SimulationSpace, PrevCSTM, CompWorldSpaceTM, RootBoneTM);
+								const FTransform NextSSTM = ConvertCSTransformToSimSpace(SimulationSpace, NextCSTM, CompWorldSpaceTM, RootBoneTM);
+
+								// Linear Velocity
+								BodyData.TransferedBoneVelocity.SetTranslation((NextSSTM.GetLocation() - PrevSSTM.GetLocation()) / DeltaSeconds);
+
+								// Angular Velocity
+								const FQuat DeltaRotation = (NextSSTM.GetRotation().Inverse() * PrevSSTM.GetRotation());
+								const float RotationAngle = DeltaRotation.GetAngle() / DeltaSeconds;
+								BodyData.TransferedBoneVelocity.SetRotation(FQuat(DeltaRotation.GetRotationAxis(), RotationAngle)); 
+							}
+						}
+					}
 				}
 
-				if(bResetSimulated)
+				// Initialize bodies.
 				{
-					Bodies[BodyIndex]->SetWorldTransform(BodyTM);
+					for (const FOutputBoneData& OutputData : OutputBoneData)
+					{
+						const int32 BodyIndex = OutputData.BodyIndex;
+						BodyAnimData[BodyIndex].bBodyTransformInitialized = true;
+
+						const FTransform& ComponentSpaceTM = Output.Pose.GetComponentSpaceTransform(OutputData.CompactPoseBoneIndex);
+						const FTransform BodyTM = ConvertCSTransformToSimSpace(SimulationSpace, ComponentSpaceTM, CompWorldSpaceTM, RootBoneTM);
+						Bodies[BodyIndex]->SetWorldTransform(BodyTM);
+					}
 				}
-				else
+			}
+			// Subsequent updates, update kinematic bodies and run simulation.
+			else
+			{
+				// Transfer bone velocities previously captured.
+				if (bTransferBoneVelocities && (CapturedBoneVelocityPose.GetPose().GetNumBones() > 0))
 				{
-					Bodies[BodyIndex]->SetKinematicTarget(BodyTM);
+					for (const FOutputBoneData& OutputData : OutputBoneData)
+					{
+						const int32 BodyIndex = OutputData.BodyIndex;
+						const FBodyAnimData& BodyData = BodyAnimData[BodyIndex];
+
+						if (BodyData.bIsSimulated)
+						{
+							ImmediatePhysics::FActorHandle* Body = Bodies[BodyIndex];
+							Body->SetLinearVelocity(BodyData.TransferedBoneVelocity.GetTranslation());
+
+							const FQuat AngularVelocity = BodyData.TransferedBoneVelocity.GetRotation();
+							Body->SetAngularVelocity(AngularVelocity.GetRotationAxis() * AngularVelocity.GetAngle());
+						}
+					}
+
+					// Free up our captured pose after it's been used.
+					CapturedBoneVelocityPose.Empty();
 				}
-				
-			}
-		}
 
-		bResetSimulated = false;
+				for (const FOutputBoneData& OutputData : OutputBoneData)
+				{
+					const int32 BodyIndex = OutputData.BodyIndex;
+					if (!BodyAnimData[BodyIndex].bIsSimulated)
+					{
+						const FTransform& ComponentSpaceTM = Output.Pose.GetComponentSpaceTransform(OutputData.CompactPoseBoneIndex);
+						const FTransform BodyTM = ConvertCSTransformToSimSpace(SimulationSpace, ComponentSpaceTM, CompWorldSpaceTM, RootBoneTM);
+
+						Bodies[BodyIndex]->SetKinematicTarget(BodyTM);
+					}
+				}
 
-		UpdateWorldForces(CompWorldSpaceTM, RootBoneTM);
+				UpdateWorldForces(CompWorldSpaceTM, RootBoneTM);
+				const FVector SimSpaceGravity = WorldVectorToSpaceNoScale(SimulationSpace, WorldSpaceGravity, CompWorldSpaceTM, RootBoneTM);
 
-		//simulate
-		const FVector SimSpaceGravity = WorldVectorToSpaceNoScale(SimulationSpace, WorldSpaceGravity, CompWorldSpaceTM, RootBoneTM);
-		PhysicsSimulation->Simulate(DeltaSeconds, SimSpaceGravity);
+				// Run simulation at a minimum of 30 FPS to prevent system from exploding.
+				// DeltaTime can be higher due to URO, so take multiple iterations in that case.
+				const float MaxDeltaSeconds = 1.f / 30.f;
+				const int32 NumIterations = FMath::Clamp(FMath::CeilToInt(DeltaSeconds / MaxDeltaSeconds), 1, 4);
+				const float StepDeltaTime = DeltaSeconds / float(NumIterations);
 
+				for (int32 Step = 1; Step <= NumIterations; Step++)
+				{
+					PhysicsSimulation->Simulate(StepDeltaTime, SimSpaceGravity);
+				}
+			}
+		}
 		
 		//write back to animation system
 		for (const FOutputBoneData& OutputData : OutputBoneData)
 		{
-			int32 BodyIndex = OutputData.BodyIndex;
-			if(BodyIndex != INDEX_NONE)
+			const int32 BodyIndex = OutputData.BodyIndex;
+			if (BodyAnimData[BodyIndex].bIsSimulated)
 			{
-				/*if(true || IsSimulated[BodyIndex])*/	//todo: only output kinematic bones that have simulated ancestors
-				{
-					FCompactPoseBoneIndex SimBoneIndex = OutputData.BoneReference.GetCompactPoseIndex(BoneContainer);
-					const FTransform BodyTM = Bodies[BodyIndex]->GetWorldTransform();
-					FTransform ComponentSpaceTM;
+				const FTransform BodyTM = Bodies[BodyIndex]->GetWorldTransform();
+				FTransform ComponentSpaceTM;
 
-					switch(SimulationSpace)
-					{
-						case ESimulationSpace::ComponentSpace: ComponentSpaceTM = BodyTM; break;
-						case ESimulationSpace::WorldSpace: ComponentSpaceTM = BodyTM.GetRelativeTransform(CompWorldSpaceTM); break;
-						case ESimulationSpace::RootBoneSpace: ComponentSpaceTM = BodyTM * RootBoneTM; break;
-						default: ensureMsgf(false, TEXT("Unsupported Simulation Space")); ComponentSpaceTM = BodyTM;
-					}
-					
-					OutBoneTransforms.Add(FBoneTransform(SimBoneIndex, ComponentSpaceTM));
+				switch(SimulationSpace)
+				{
+					case ESimulationSpace::ComponentSpace: ComponentSpaceTM = BodyTM; break;
+					case ESimulationSpace::WorldSpace: ComponentSpaceTM = BodyTM.GetRelativeTransform(CompWorldSpaceTM); break;
+					case ESimulationSpace::RootBoneSpace: ComponentSpaceTM = BodyTM * RootBoneTM; break;
+					default: ensureMsgf(false, TEXT("Unsupported Simulation Space")); ComponentSpaceTM = BodyTM;
 				}
-			}
-			else
-			{
-				//we have no body, but our ancestors are simulated so update our component space transform
+					
+				OutBoneTransforms.Add(FBoneTransform(OutputData.CompactPoseBoneIndex, ComponentSpaceTM));
 			}
 		}
 	}
@@ -212,20 +400,34 @@ void ComputeBodyInsertionOrder(TArray<FBoneIndexType>& InsertionOrder, const USk
 void FAnimNode_RigidBody::InitPhysics(const UAnimInstance* InAnimInstance)
 {
 	const USkeletalMeshComponent* SkeletalMeshComp = InAnimInstance->GetSkelMeshComponent();
-	const FReferenceSkeleton& RefSkel = SkeletalMeshComp->SkeletalMesh->RefSkeleton;
+	const USkeletalMesh* SkeletalMeshAsset = SkeletalMeshComp->SkeletalMesh;
+
+	const FReferenceSkeleton& SkelMeshRefSkel = SkeletalMeshAsset->RefSkeleton;
 	UPhysicsAsset* UsePhysicsAsset = OverridePhysicsAsset ? OverridePhysicsAsset : InAnimInstance->GetSkelMeshComponent()->GetPhysicsAsset();
+		
+	USkeleton* SkeletonAsset = InAnimInstance->CurrentSkeleton;
+	ensure(SkeletonAsset == SkeletalMeshAsset->Skeleton);
+
+	const int32 SkelMeshLinkupIndex = SkeletonAsset->GetMeshLinkupIndex(SkeletalMeshAsset);
+	ensure(SkelMeshLinkupIndex != INDEX_NONE);
+	const FSkeletonToMeshLinkup& SkeletonToMeshLinkupTable = SkeletonAsset->LinkupCache[SkelMeshLinkupIndex];
+	const TArray<int32>& MeshToSkeletonBoneIndex = SkeletonToMeshLinkupTable.MeshToSkeletonTable;
+	
+	const int32 NumSkeletonBones = SkeletonAsset->GetReferenceSkeleton().GetNum();
+	SkeletonBoneIndexToBodyIndex.Reset(NumSkeletonBones);
+	SkeletonBoneIndexToBodyIndex.Init(INDEX_NONE, NumSkeletonBones);
+
 	if(UsePhysicsAsset)
 	{
 		delete PhysicsSimulation;
 		PhysicsSimulation = new FSimulation();
 		const int32 NumBodies = UsePhysicsAsset->SkeletalBodySetups.Num();
 		Bodies.Empty(NumBodies);
-		BodyBoneIndices.Empty(NumBodies);
 		ComponentsInSim.Reset();
-		IsSimulated.Reset();
+		BodyAnimData.Reset(NumBodies);
+		BodyAnimData.AddDefaulted(NumBodies);
 		TotalMass = 0.f;
 		
-
 		TArray<FBodyInstance*> HighLevelBodyInstances;
 		TArray<FConstraintInstance*> HighLevelConstraintInstances;
 		SkeletalMeshComp->InstantiatePhysicsAsset(*UsePhysicsAsset, SimulationSpace == ESimulationSpace::WorldSpace ? SkeletalMeshComp->GetComponentToWorld().GetScale3D() : FVector(1.f), HighLevelBodyInstances, HighLevelConstraintInstances);
@@ -281,10 +483,12 @@ void FAnimNode_RigidBody::InitPhysics(const UAnimInstance* InAnimInstance)
 
 					if (NewBodyHandle)
 					{
-						Bodies.Add(NewBodyHandle);
-						BodyBoneIndices.Add(InsertBone);
+						const int32 BodyIndex = Bodies.Add(NewBodyHandle);
+
+						const int32 SkeletonBoneIndex = MeshToSkeletonBoneIndex[InsertBone];
+						SkeletonBoneIndexToBodyIndex[SkeletonBoneIndex] = BodyIndex;
+						BodyAnimData[BodyIndex].bIsSimulated = !bKinematic;
 						NamesToHandles.Add(BodySetup->BoneName, NewBodyHandle);
-						IsSimulated.Add(!bKinematic);
 						BodyIndexToActorHandle[BodyInstance->InstanceBodyIndex] = NewBodyHandle;
 
 						if (BodySetup->CollisionReponse == EBodyCollisionResponse::BodyCollision_Disabled)
@@ -316,15 +520,15 @@ void FAnimNode_RigidBody::InitPhysics(const UAnimInstance* InAnimInstance)
 			InsertionOrderPerBone[InsertionOrder[Position]] = Position;
 		}
 
-		HighLevelConstraintInstances.Sort([&InsertionOrderPerBone, &RefSkel](const FConstraintInstance& LHS, const FConstraintInstance& RHS)
+		HighLevelConstraintInstances.Sort([&InsertionOrderPerBone, &SkelMeshRefSkel](const FConstraintInstance& LHS, const FConstraintInstance& RHS)
 		{
 			if(LHS.IsValidConstraintInstance() && RHS.IsValidConstraintInstance())
 			{
-				const int32 BoneIdxLHS1 = RefSkel.FindBoneIndex(LHS.ConstraintBone1);
-				const int32 BoneIdxLHS2 = RefSkel.FindBoneIndex(LHS.ConstraintBone2);
+				const int32 BoneIdxLHS1 = SkelMeshRefSkel.FindBoneIndex(LHS.ConstraintBone1);
+				const int32 BoneIdxLHS2 = SkelMeshRefSkel.FindBoneIndex(LHS.ConstraintBone2);
 
-				const int32 BoneIdxRHS1 = RefSkel.FindBoneIndex(RHS.ConstraintBone1);
-				const int32 BoneIdxRHS2 = RefSkel.FindBoneIndex(RHS.ConstraintBone2);
+				const int32 BoneIdxRHS1 = SkelMeshRefSkel.FindBoneIndex(RHS.ConstraintBone1);
+				const int32 BoneIdxRHS2 = SkelMeshRefSkel.FindBoneIndex(RHS.ConstraintBone2);
 
 				const int32 MaxPositionLHS = FMath::Max(InsertionOrderPerBone[BoneIdxLHS1], InsertionOrderPerBone[BoneIdxLHS2]);
 				const int32 MaxPositionRHS = FMath::Max(InsertionOrderPerBone[BoneIdxRHS1], InsertionOrderPerBone[BoneIdxRHS2]);
@@ -456,10 +660,16 @@ void FAnimNode_RigidBody::UpdateWorldForces(const FTransform& ComponentToWorld,
 	}
 }
 
+DECLARE_CYCLE_STAT(TEXT("RigidBody_PreUpdate"), STAT_RigidBody_PreUpdate, STATGROUP_Anim);
+
 void FAnimNode_RigidBody::PreUpdate(const UAnimInstance* InAnimInstance)
 {
+	SCOPE_CYCLE_COUNTER(STAT_RigidBody_PreUpdate);
+
 	UWorld* World = InAnimInstance->GetWorld();
 	USkeletalMeshComponent* SKC = InAnimInstance->GetSkelMeshComponent();
+	APawn* PawnOwner = InAnimInstance->TryGetPawnOwner();
+	UPawnMovementComponent* MovementComp = PawnOwner ? PawnOwner->GetMovementComponent() : nullptr;
 
 #if WITH_EDITOR
 	if (bEnableWorldGeometry && SimulationSpace != ESimulationSpace::WorldSpace)
@@ -469,9 +679,7 @@ void FAnimNode_RigidBody::PreUpdate(const UAnimInstance* InAnimInstance)
 	}
 #endif
 
-	DeltaSeconds = World->GetDeltaSeconds();
-	WorldSpaceGravity = bOverrideWorldGravity ? OverrideWorldGravity : FVector(0.f, 0.f, World->GetGravityZ());
-
+	WorldSpaceGravity = bOverrideWorldGravity ? OverrideWorldGravity : (MovementComp ? FVector(0.f, 0.f, MovementComp->GetGravityZ()) : FVector(0.f, 0.f, World->GetGravityZ()));
 	if(SKC)
 	{
 		if (PhysicsSimulation && bEnableWorldGeometry && SimulationSpace == ESimulationSpace::WorldSpace)
@@ -480,16 +688,24 @@ void FAnimNode_RigidBody::PreUpdate(const UAnimInstance* InAnimInstance)
 		}
 
 		PendingRadialForces = SKC->GetPendingRadialForces();
-	}
-	
+	}	
 }
 
+DECLARE_CYCLE_STAT(TEXT("RigidBody_Update"), STAT_RigidBody_Update, STATGROUP_Anim);
+
 void FAnimNode_RigidBody::UpdateInternal(const FAnimationUpdateContext& Context)
 {
+	SCOPE_CYCLE_COUNTER(STAT_RigidBody_PreUpdate);
+
+	// Accumulate deltatime elapsed during update. To be used during evaluation.
+	AccumulatedDeltaTime += Context.AnimInstanceProxy->GetDeltaSeconds();
+
 	if (UnsafeWorld != nullptr)
-	{
+	{		
 
 #if WITH_PHYSX
+		// Node is valid to evaluate. Simulation is starting.
+		bSimulationStarted = true;
 
 		TArray<FOverlapResult> Overlaps;
 		UnsafeWorld->OverlapMultiByChannel(Overlaps, Bounds.Center, FQuat::Identity, OverlapChannel, FCollisionShape::MakeSphere(Bounds.W), QueryParams, FCollisionResponseParams(ECR_Overlap));
@@ -523,11 +739,8 @@ void FAnimNode_RigidBody::InitializeBoneReferences(const FBoneContainer& Require
 	const TArray<FBoneIndexType>& RequiredBoneIndices = RequiredBones.GetBoneIndicesArray();
 	const int32 NumRequiredBoneIndices = RequiredBoneIndices.Num();
 	const FReferenceSkeleton& RefSkeleton = RequiredBones.GetReferenceSkeleton();
-	
-	OutputBoneData.Empty(NumRequiredBoneIndices);
 
-	TArray<bool> OutputBonesCache;	//easy lookup into whether parent is going to be evaluated
-	OutputBonesCache.AddZeroed(RefSkeleton.GetNum());
+	OutputBoneData.Empty(NumBodies);
 
 	int32 NumSimulatedBodies = 0;
 
@@ -535,52 +748,56 @@ void FAnimNode_RigidBody::InitializeBoneReferences(const FBoneContainer& Require
 	RootBoneRef.BoneName = RefSkeleton.GetBoneName(0);
 	RootBoneRef.Initialize(RequiredBones);
 
-	for(int32 RequiredBoneIndexIndex = 0; RequiredBoneIndexIndex < NumRequiredBoneIndices; ++RequiredBoneIndexIndex)
+	for(int32 Index = 0; Index < NumRequiredBoneIndices; ++Index)
 	{
-		const FBoneIndexType RequiredBoneIndex = RequiredBoneIndices[RequiredBoneIndexIndex];
-		int32 FoundBodyIndex = INDEX_NONE;
+		const FCompactPoseBoneIndex CompactPoseBoneIndex(Index);
+		const FBoneIndexType SkeletonBoneIndex = RequiredBones.GetSkeletonIndex(CompactPoseBoneIndex);
+		const int32 BodyIndex = SkeletonBoneIndexToBodyIndex[SkeletonBoneIndex];
 
-		for(int32 BodyIndex = 0; BodyIndex < NumBodies; ++BodyIndex)	//TODO: we could sort this and avoid n^2
+		if (BodyIndex != INDEX_NONE)
 		{
-			if(BodyBoneIndices[BodyIndex] == RequiredBoneIndex)	//If we have a body we need to save it for later
-			{
-				FOutputBoneData* OutputData = new (OutputBoneData) FOutputBoneData();
-				OutputData->BodyIndex = BodyIndex;
-				OutputData->BoneReference.BoneName = RefSkeleton.GetBoneName(RequiredBoneIndex);
-				OutputData->BoneReference.Initialize(RequiredBones);
+			//If we have a body we need to save it for later
+			FOutputBoneData* OutputData = new (OutputBoneData) FOutputBoneData();
+			OutputData->BodyIndex = BodyIndex;
+			OutputData->CompactPoseBoneIndex = CompactPoseBoneIndex;
 
-				if(IsSimulated[BodyIndex])
-				{
-					OutputBonesCache[RequiredBoneIndex] = true;	//children of simulated bodies need to update the component space transform
-					++NumSimulatedBodies;
-				}
-		
-				//++BodyIndex;	//Move on to next body
-				FoundBodyIndex = BodyIndex;
-				break;
+			if (BodyAnimData[BodyIndex].bIsSimulated)
+			{
+				++NumSimulatedBodies;
 			}
-		}
 
-		if(RequiredBoneIndex > 0 && FoundBodyIndex == INDEX_NONE)
-		{
-			//If we don't have a body, but our ancestors are simulated we will need to update the component space transform
-			FBoneIndexType ParentBoneIndex = RequiredBones.GetParentBoneIndex(RequiredBoneIndex);
-			if (OutputBonesCache[ParentBoneIndex])
+			OutputData->BoneIndicesToParentBody.Add(CompactPoseBoneIndex);
+
+			// Walk up parent chain until we find parent body.
+			OutputData->ParentBodyIndex = INDEX_NONE;
+			FCompactPoseBoneIndex CompactParentIndex = RequiredBones.GetParentBoneIndex(CompactPoseBoneIndex);
+			while (CompactParentIndex != INDEX_NONE)
 			{
-				OutputBonesCache[RequiredBoneIndex] = true;
+				const FBoneIndexType SkeletonParentBoneIndex = RequiredBones.GetSkeletonIndex(CompactParentIndex);
+				OutputData->ParentBodyIndex = SkeletonBoneIndexToBodyIndex[SkeletonParentBoneIndex];
+				if (OutputData->ParentBodyIndex != INDEX_NONE)
+				{
+					break;
+				}
 
-				FOutputBoneData* OutputData = new (OutputBoneData) FOutputBoneData();
-				OutputData->BodyIndex = INDEX_NONE;
-				OutputData->BoneReference.BoneName = RefSkeleton.GetBoneName(RequiredBoneIndex);
-				OutputData->BoneReference.Initialize(RequiredBones);
+				OutputData->BoneIndicesToParentBody.Add(CompactParentIndex);
+				CompactParentIndex = RequiredBones.GetParentBoneIndex(CompactParentIndex);
 			}
 		}
 	}
 
+	// New bodies protentially introduced with new LOD
+	// We'll have to initialize their transform.
+	bCheckForBodyTransformInit = true;
+
 	if(PhysicsSimulation)
 	{
 		PhysicsSimulation->SetNumActiveBodies(NumSimulatedBodies);
 	}
+
+	// We're switching to a new LOD, this invalidates our captured poses.
+	CapturedFrozenPose.Empty();
+	CapturedFrozenCurves.Empty();
 }
 
 void FAnimNode_RigidBody::OnInitializeAnimInstance(const FAnimInstanceProxy* InProxy, const UAnimInstance* InAnimInstance)
diff --git a/Engine/Plugins/Runtime/ImmediatePhysics/Source/ImmediatePhysics/Public/BoneControllers/AnimNode_RigidBody.h b/Engine/Plugins/Runtime/ImmediatePhysics/Source/ImmediatePhysics/Public/BoneControllers/AnimNode_RigidBody.h
--- a/Engine/Plugins/Runtime/ImmediatePhysics/Source/ImmediatePhysics/Public/BoneControllers/AnimNode_RigidBody.h
+++ b/Engine/Plugins/Runtime/ImmediatePhysics/Source/ImmediatePhysics/Public/BoneControllers/AnimNode_RigidBody.h
@@ -43,6 +43,8 @@ struct IMMEDIATEPHYSICS_API FAnimNode_RigidBody : public FAnimNode_SkeletalContr
 	// End of FAnimNode_Base interface
 
 	// FAnimNode_SkeletalControlBase interface
+	virtual void UpdateComponentPose_AnyThread(const FAnimationUpdateContext& Context) override;
+	virtual void EvaluateComponentPose_AnyThread(FComponentSpacePoseContext& Output) override;
 	virtual void EvaluateSkeletalControl_AnyThread(FComponentSpacePoseContext& Output, TArray<FBoneTransform>& OutBoneTransforms) override;
 	virtual void OnInitializeAnimInstance(const FAnimInstanceProxy* InProxy, const UAnimInstance* InAnimInstance) override;
 	virtual void PreUpdate(const UAnimInstance* InAnimInstance) override;
@@ -85,6 +87,21 @@ struct IMMEDIATEPHYSICS_API FAnimNode_RigidBody : public FAnimNode_SkeletalContr
 	UPROPERTY(EditAnywhere, Category = Settings, meta = (ClampMin="1.0", ClampMax="2.0"))
 	float CachedBoundsScale;
 
+	/** 
+		When simulation starts, transfer previous bone velocities (from animation)
+		to make transition into simulation seamless.
+	*/
+	UPROPERTY(EditAnywhere, Category = Settings, meta=(PinHiddenByDefault))
+	bool bTransferBoneVelocities;
+
+	/**
+		When simulation starts, freeze incoming pose.
+		This is useful for ragdolls, when we want the simulation to take over.
+		It prevents non simulated bones from animating.
+	*/
+	UPROPERTY(EditAnywhere, Category = Settings)
+	bool bFreezeIncomingPoseOnStart;
+
 	void PostSerialize(const FArchive& Ar);
 
 private:
@@ -99,6 +116,9 @@ private:
 	void InitPhysics(const UAnimInstance* InAnimInstance);
 	void UpdateWorldGeometry(const UWorld& World, const USkeletalMeshComponent& SKC);
 	void UpdateWorldForces(const FTransform& ComponentToWorld, const FTransform& RootBoneTM);
+
+	void InitializeNewBodyTransformsDuringSimulation(FComponentSpacePoseContext& Output, const FTransform& ComponentTransform, const FTransform& RootBoneTM);
+
 private:
 
 	/** This should only be used for removing the delegate during termination. Do NOT use this for any per frame work */
@@ -108,25 +128,42 @@ private:
 
 	struct FOutputBoneData
 	{
-		FBoneReference BoneReference;
+		FOutputBoneData()
+			: CompactPoseBoneIndex(INDEX_NONE)
+		{}
+
+		FCompactPoseBoneIndex CompactPoseBoneIndex;
 		int32 BodyIndex;
+		int32 ParentBodyIndex;
+		TArray<FCompactPoseBoneIndex> BoneIndicesToParentBody;
+	};
+
+	struct FBodyAnimData
+	{
+		FBodyAnimData()
+			: bIsSimulated(false)
+			, bBodyTransformInitialized(false)
+		{}
+
+		bool bIsSimulated;
+		bool bBodyTransformInitialized;
+		FTransform TransferedBoneVelocity;
 	};
 	
 	FBoneReference RootBoneRef;
 
 	TArray<FOutputBoneData> OutputBoneData;
 	TArray<ImmediatePhysics::FActorHandle*> Bodies;
-	TArray<bool> IsSimulated;
-	TArray<FBoneIndexType> BodyBoneIndices;
-	bool bResetSimulated;
-	
+	TArray<int32> SkeletonBoneIndexToBodyIndex;
+	TArray<FBodyAnimData> BodyAnimData;
+
 	TArray<struct FPhysicsConstraintHandle*> Constraints;
 	TArray<USkeletalMeshComponent::FPendingRadialForces> PendingRadialForces;
 
 	TSet<UPrimitiveComponent*> ComponentsInSim;
 
 	FVector WorldSpaceGravity;
-	float DeltaSeconds;
+	float AccumulatedDeltaTime;
 	float TotalMass;
 
 	FSphere Bounds;
@@ -139,6 +176,15 @@ private:
 	// Typically, World should never be accessed off the Game Thread.
 	// However, since we're just doing overlaps this should be OK.
 	const UWorld* UnsafeWorld;
+
+	FBoneContainer CapturedBoneVelocityBoneContainer;
+	FCSPose<FCompactHeapPose> CapturedBoneVelocityPose;
+	FCSPose<FCompactHeapPose> CapturedFrozenPose;
+	FBlendedHeapCurve CapturedFrozenCurves;
+
+	bool bResetSimulated;
+	bool bSimulationStarted;
+	bool bCheckForBodyTransformInit;
 };
 
 template<>
diff --git a/Engine/Plugins/Runtime/PacketHandlers/CompressionComponents/Oodle/ExtractAllArchives.bat b/Engine/Plugins/Runtime/PacketHandlers/CompressionComponents/Oodle/ExtractAllArchives.bat
--- a/Engine/Plugins/Runtime/PacketHandlers/CompressionComponents/Oodle/ExtractAllArchives.bat
+++ b/Engine/Plugins/Runtime/PacketHandlers/CompressionComponents/Oodle/ExtractAllArchives.bat
@@ -13,11 +13,11 @@ set /p Directory=Type the root directory of the zipped captures:
 echo.
 
 :getoptions
-if "%2" == "" GOTO continue
+if "%2" == "all" GOTO continue
 set Options= --exclude "*" --include "*%2*.ucap.gz"
 
 :continue
-aws s3 sync %3 %Directory% %Options%
+call aws s3 sync %3 %Directory% %Options%
 
 pushd %CD%
 cd %Directory%
diff --git a/Engine/Plugins/Runtime/PacketHandlers/CompressionComponents/Oodle/Source/OodleHandlerComponent/OodleHandlerComponent.Build.cs b/Engine/Plugins/Runtime/PacketHandlers/CompressionComponents/Oodle/Source/OodleHandlerComponent/OodleHandlerComponent.Build.cs
--- a/Engine/Plugins/Runtime/PacketHandlers/CompressionComponents/Oodle/Source/OodleHandlerComponent/OodleHandlerComponent.Build.cs
+++ b/Engine/Plugins/Runtime/PacketHandlers/CompressionComponents/Oodle/Source/OodleHandlerComponent/OodleHandlerComponent.Build.cs
@@ -17,7 +17,8 @@ public class OodleHandlerComponent : ModuleRules
 				"PacketHandler",
 				"Core",
 				"CoreUObject",
-				"Engine"
+				"Engine",
+                "Analytics"
 			});
 
 
diff --git a/Engine/Plugins/Runtime/PacketHandlers/CompressionComponents/Oodle/Source/OodleHandlerComponent/Private/OodleAnalytics.cpp b/Engine/Plugins/Runtime/PacketHandlers/CompressionComponents/Oodle/Source/OodleHandlerComponent/Private/OodleAnalytics.cpp
new file mode 100644
--- /dev/null
+++ b/Engine/Plugins/Runtime/PacketHandlers/CompressionComponents/Oodle/Source/OodleHandlerComponent/Private/OodleAnalytics.cpp
@@ -0,0 +1,50 @@
+// Copyright 1998-2016 Epic Games, Inc. All Rights Reserved.
+
+#include "OodleAnalytics.h"
+#include "OodleHandlerComponent.h"
+#include "AnalyticsEventAttribute.h"
+#include "IAnalyticsProvider.h"
+
+class FOodleAnalyticsImpl
+	: public FOodleAnalytics
+{
+public:
+
+	virtual ~FOodleAnalyticsImpl() { }
+
+	virtual void SetProvider(TSharedPtr<IAnalyticsProvider> InAnalyticsProvider) override
+	{
+		AnalyticsProvider = InAnalyticsProvider;
+	}
+
+	virtual void FireEvent_OodleAnalytics(float InTotalSavings, float OutTotalSavings) const override
+	{
+		if (AnalyticsProvider.IsValid())
+		{
+			static const FString EventName = TEXT("Oodle.Stats");
+
+			TArray<FAnalyticsEventAttribute> Attributes;
+			Attributes.Add(FAnalyticsEventAttribute(TEXT("InTotalSavings"), InTotalSavings));
+			Attributes.Add(FAnalyticsEventAttribute(TEXT("OutTotalSavings"), OutTotalSavings));
+			AnalyticsProvider->RecordEvent(EventName, Attributes);
+		}
+	}
+
+private: // functions
+
+	FOodleAnalyticsImpl()
+	{}
+
+private: // data
+
+	// cached analytics provider for pushing events
+	TSharedPtr<IAnalyticsProvider> AnalyticsProvider;
+
+	friend FOodleAnalyticsFactory;
+};
+
+TSharedRef< FOodleAnalytics > FOodleAnalyticsFactory::Create()
+{
+	TSharedRef< FOodleAnalytics > Instance = MakeShareable(new FOodleAnalyticsImpl());
+	return Instance;
+}
diff --git a/Engine/Plugins/Runtime/PacketHandlers/CompressionComponents/Oodle/Source/OodleHandlerComponent/Private/OodleHandlerComponent.cpp b/Engine/Plugins/Runtime/PacketHandlers/CompressionComponents/Oodle/Source/OodleHandlerComponent/Private/OodleHandlerComponent.cpp
--- a/Engine/Plugins/Runtime/PacketHandlers/CompressionComponents/Oodle/Source/OodleHandlerComponent/Private/OodleHandlerComponent.cpp
+++ b/Engine/Plugins/Runtime/PacketHandlers/CompressionComponents/Oodle/Source/OodleHandlerComponent/Private/OodleHandlerComponent.cpp
@@ -11,6 +11,7 @@
 #include "Misc/EngineVersion.h"
 
 #include "OodleTrainerCommandlet.h"
+#include "OodleAnalytics.h"
 
 #if !UE_BUILD_SHIPPING
 #include "EngineGlobals.h"
@@ -275,16 +276,40 @@ OodleHandlerComponent::OodleHandlerComponent()
 	, OodleReservedPacketBits(0)
 	, ServerDictionary()
 	, ClientDictionary()
+	, TotalInCompressedLength(0)
+	, TotalInDecompressedLength(0)
+	, TotalOutCompressedLength(0)
+	, TotalOutUncompressedLength(0)
 {
 	SetActive(true);
 
 #if !UE_BUILD_SHIPPING || OODLE_DEV_SHIPPING
 	OodleComponentList.Add(this);
 #endif
+
+	Analytics = FOodleAnalyticsFactory::Create();
 }
 
 OodleHandlerComponent::~OodleHandlerComponent()
 {
+	float InTotalSavings = 0.0f;
+	float OutTotalSavings = 0.0f;
+
+	if (TotalInCompressedLength > 0)
+	{
+		InTotalSavings = (1.0 - ((double)TotalInCompressedLength / (double)TotalInDecompressedLength)) * 100.0;
+	}
+
+	if (TotalOutCompressedLength > 0)
+	{
+		OutTotalSavings = (1.0 - ((double)TotalOutCompressedLength / (double)TotalOutUncompressedLength)) * 100.0;
+	}
+
+	if (Analytics.IsValid())
+	{
+		Analytics->FireEvent_OodleAnalytics(InTotalSavings, OutTotalSavings);
+	}
+
 #if !UE_BUILD_SHIPPING || OODLE_DEV_SHIPPING
 	OodleComponentList.Remove(this);
 
@@ -852,6 +877,8 @@ void OodleHandlerComponent::Incoming(FBitReader& Packet)
 #if STATS
 						GOodleNetStats.IncomingStats(CompressedLength, DecompressedLength);
 #endif
+						TotalInCompressedLength += CompressedLength;
+						TotalInDecompressedLength += DecompressedLength;
 					}
 					else
 					{
@@ -971,6 +998,9 @@ void OodleHandlerComponent::Outgoing(FBitWriter& Packet)
 #if STATS
 						GOodleNetStats.OutgoingStats(CompressedBytes, UncompressedBytes);
 #endif
+						TotalOutCompressedLength += CompressedBytes;
+						TotalOutUncompressedLength += UncompressedBytes;
+
 					}
 					else
 					{
@@ -978,22 +1008,28 @@ void OodleHandlerComponent::Outgoing(FBitWriter& Packet)
 						//					at the start of the bit writer
 						Packet.SerializeBits(UncompressedData, UncompressedBits);
 
-#if STATS
+
 						if (!bWithinBitBounds)
 						{
+#if STATS
 							INC_DWORD_STAT(STAT_Oodle_CompressFailSize);
+#endif
 						}
 						else if (CompressedBytes >= UncompressedBytes)
 						{
+#if STATS
 							GOodleNetStats.OutgoingStats(UncompressedBytes, UncompressedBytes);
 
 							INC_DWORD_STAT(STAT_Oodle_CompressFailSavings);
+#endif
+
+							TotalOutCompressedLength += CompressedBytes;
+							TotalOutUncompressedLength += UncompressedBytes;
 						}
 						else
 						{
 							// @todo #JohnB
 						}
-#endif
 					}
 				}
 				else
@@ -1076,6 +1112,13 @@ int32 OodleHandlerComponent::GetReservedPacketBits()
 	return ReturnVal;
 }
 
+void OodleHandlerComponent::SetAnalyticsProvider(TSharedPtr<class IAnalyticsProvider> Provider)
+{
+	if (Analytics.IsValid())
+	{
+		Analytics->SetProvider(Provider);
+	}
+}
 
 #if !UE_BUILD_SHIPPING
 /**
diff --git a/Engine/Plugins/Runtime/PacketHandlers/CompressionComponents/Oodle/Source/OodleHandlerComponent/Public/OodleAnalytics.h b/Engine/Plugins/Runtime/PacketHandlers/CompressionComponents/Oodle/Source/OodleHandlerComponent/Public/OodleAnalytics.h
new file mode 100644
--- /dev/null
+++ b/Engine/Plugins/Runtime/PacketHandlers/CompressionComponents/Oodle/Source/OodleHandlerComponent/Public/OodleAnalytics.h
@@ -0,0 +1,38 @@
+// Copyright 1998-2017 Epic Games, Inc. All Rights Reserved.
+#pragma once
+
+#include "IAnalyticsProvider.h"
+
+#define FACTORY(ReturnType, Type, ...) \
+class Type##Factory \
+{ \
+public: \
+	static ReturnType Create(__VA_ARGS__); \
+}; 
+
+class FOodleAnalytics
+{
+public:
+
+	virtual ~FOodleAnalytics() { }
+
+	/**
+	* Update provider to use for capturing events
+	*/
+	virtual void SetProvider(TSharedPtr<IAnalyticsProvider> InAnalyticsProvider) = 0;
+
+	/*
+	* @EventName Oodle.Stats
+	*
+	* @Trigger Fires event for percentage compression by Oodle
+	*
+	* @EventParam InTotalSavings Percentage savings in traffic sent by a client
+	* @EventParam OutTotalSavings Percentage savings in traffic sent by the server
+	*
+	* @Comments
+	*
+	*/
+	virtual void FireEvent_OodleAnalytics(float InTotalSavings, float OutTotalSavings) const = 0;
+};
+
+FACTORY(TSharedRef<class FOodleAnalytics>, FOodleAnalytics);
diff --git a/Engine/Plugins/Runtime/PacketHandlers/CompressionComponents/Oodle/Source/OodleHandlerComponent/Public/OodleHandlerComponent.h b/Engine/Plugins/Runtime/PacketHandlers/CompressionComponents/Oodle/Source/OodleHandlerComponent/Public/OodleHandlerComponent.h
--- a/Engine/Plugins/Runtime/PacketHandlers/CompressionComponents/Oodle/Source/OodleHandlerComponent/Public/OodleHandlerComponent.h
+++ b/Engine/Plugins/Runtime/PacketHandlers/CompressionComponents/Oodle/Source/OodleHandlerComponent/Public/OodleHandlerComponent.h
@@ -351,6 +351,9 @@ public:
 
 	virtual int32 GetReservedPacketBits() override;
 
+	/** Sets the analytics provider that's used for sending Oodle performance */
+	virtual void SetAnalyticsProvider(TSharedPtr<class IAnalyticsProvider> Provider);
+
 protected:
 	/** Whether or not Oodle is enabled */
 	bool bEnableOodle;
@@ -372,6 +375,8 @@ protected:
 	/** Cached reserved packet bits for Oodle */
 	uint32 OodleReservedPacketBits;
 
+	/** The analytics provider that's used for sending Oodle performance */
+	TSharedPtr<class FOodleAnalytics> Analytics;
 
 #if !UE_BUILD_SHIPPING
 public:
@@ -382,6 +387,19 @@ public:
 
 	/** Client (Incoming - relative to server) dictionary data */
 	TSharedPtr<FOodleDictionary> ClientDictionary;
+
+	/** Input traffic compressed packet length */
+	uint32 TotalInCompressedLength;
+
+	/** Input traffic decompressed packet length */
+	uint32 TotalInDecompressedLength;
+
+	/** Output traffic compressed packet length */
+	uint32 TotalOutCompressedLength;
+
+	/** Output traffic uncompressed packet length */
+	uint32 TotalOutUncompressedLength;
+
 };
 #endif
 
diff --git a/Engine/Plugins/Runtime/ProceduralMeshComponent/Source/ProceduralMeshComponent/Private/ProceduralMeshComponent.cpp b/Engine/Plugins/Runtime/ProceduralMeshComponent/Source/ProceduralMeshComponent/Private/ProceduralMeshComponent.cpp
--- a/Engine/Plugins/Runtime/ProceduralMeshComponent/Source/ProceduralMeshComponent/Private/ProceduralMeshComponent.cpp
+++ b/Engine/Plugins/Runtime/ProceduralMeshComponent/Source/ProceduralMeshComponent/Private/ProceduralMeshComponent.cpp
@@ -95,9 +95,14 @@ static void ConvertProcMeshToDynMeshVertex(FDynamicMeshVertex& Vert, const FProc
 }
 
 /** Procedural mesh scene proxy */
-class FProceduralMeshSceneProxy : public FPrimitiveSceneProxy
+class FProceduralMeshSceneProxy final : public FPrimitiveSceneProxy
 {
 public:
+	SIZE_T GetTypeHash() const override
+	{
+		static size_t UniquePointer;
+		return reinterpret_cast<size_t>(&UniquePointer);
+	}
 
 	FProceduralMeshSceneProxy(UProceduralMeshComponent* Component)
 		: FPrimitiveSceneProxy(Component)
diff --git a/Engine/Shaders/Private/MaterialTemplate.ush b/Engine/Shaders/Private/MaterialTemplate.ush
--- a/Engine/Shaders/Private/MaterialTemplate.ush
+++ b/Engine/Shaders/Private/MaterialTemplate.ush
@@ -100,7 +100,7 @@ struct FMaterialParticleParameters
 
 float4 GetDynamicParameter(FMaterialParticleParameters Parameters, float4 Default)
 {
-	#if PARTICLE_FACTORY
+	#if (PARTICLE_FACTORY || NIAGARA_PARTICLE_FACTORY)
 	return Parameters.DynamicParameter;
 	#else
 	return Default;
diff --git a/Engine/Shaders/Private/NiagaraEmitterInstanceShader.usf b/Engine/Shaders/Private/NiagaraEmitterInstanceShader.usf
--- a/Engine/Shaders/Private/NiagaraEmitterInstanceShader.usf
+++ b/Engine/Shaders/Private/NiagaraEmitterInstanceShader.usf
@@ -14,10 +14,39 @@ NiagaraSimulationShader.usf:
 //TODO: Add common utility code for things such as reading data sources, reading from and writing to events.
 
 #if GPU_SIMULATION
+	/** Returns clip space W, which is world space distance along the View Z axis. Note if you need DeviceZ LookupDeviceZ() is the faster option */
+	float CalcSceneDepth(float2 ScreenUV)
+	{
+	#if SCENE_TEXTURES_DISABLED
+		return 0.0f;
+	#else
+		#if FEATURE_LEVEL > FEATURE_LEVEL_ES3_1 || MOBILE_FORCE_DEPTH_TEXTURE_READS
+			return ConvertFromDeviceZ(Texture2DSampleLevel(SceneDepthTexture, SceneDepthTextureSampler, ScreenUV, 0).r);
+		#else
+			#if COMPILER_GLSL_ES2
+				#if IOS
+					// Only call FramebufferFetch when actually compiling for IOS ES2.
+					return FramebufferFetchES2().w;
+				#elif WEBGL
+					return Texture2DSampleLevel(SceneAlphaCopyTexture, SceneAlphaCopyTextureSampler, ScreenUV, 0).r;
+				#else 
+					float SceneW = ConvertFromDeviceZ(Texture2DSampleLevel(SceneDepthTexture, SceneDepthTextureSampler, ScreenUV, 0).r);		
+					return DepthbufferFetchES2(SceneW, View.InvDeviceZToWorldZTransform[2], View.InvDeviceZToWorldZTransform[3]);
+				#endif 
+			#elif METAL_PROFILE && !MAC
+					return FramebufferFetchES2().w;
+			#else
+				return ConvertFromDeviceZ(Texture2DSampleLevel(SceneDepthTexture, SceneDepthTextureSampler, ScreenUV, 0).r);
+			#endif
+		#endif
+	#endif
+	}
+
+
 	static uint3 GDispatchThreadId;
 	uint EmitterTickCounter;
 
-	#define THREADGROUP_SIZE 16
+	#define THREADGROUP_SIZE 64
 
 	// 4D random number generator inspired by PCGs (permuted congruential generator)
 	// COPY FROM RANDOM.USF, which we can't easily include here due to a sizable dependency chain
@@ -54,7 +83,7 @@ NiagaraSimulationShader.usf:
 	static uint4 RandomVal = 2308975;
 	float4 rand4()
 	{
-		RandomVal = NiagaraRand4DPCG32( int4(GDispatchThreadId+(RandomVal/57843).xyz, 54121) + int4(EmitterTickCounter.xxxx));
+		RandomVal = NiagaraRand4DPCG32( int4(GDispatchThreadId+(RandomVal/573).xyz, RandomVal.x/921) + int4(EmitterTickCounter.xxxx));
 		return float4(RandomVal.xyzw) / float4(0xFFFFFFFF,0xFFFFFFFF,0xFFFFFFFF,0xFFFFFFFF);
 	}
 
@@ -143,17 +172,21 @@ void ExitStatScope();
 #endif
 
 #if GPU_SIMULATION
-	uint NumInstancesPerThread;
-	uint StartInstance;
-	uint SimulateStartInstance;
-	uint NumInstances;
-	uint NumThreadGroups;
-	uint GroupStartInstance;
+	int ExecIndex()
+	{
+		return GDispatchThreadId.x;
+	}
+
+	uint StartInstancePhase0;
+	uint StartInstancePhase1;
+	uint TotalNumInstances;
+
 	uint ComponentBufferSizeRead;
 	uint ComponentBufferSizeWrite;
 	int	NumEventsPerParticle;
 	int NumParticlesPerEvent;
 	uint CopyInstancesBeforeStart;
+
 	Buffer<float> InputFloat;
 	Buffer<int> InputInt;
 	RWBuffer<int> RWOutputInt;
@@ -165,7 +198,7 @@ void ExitStatScope();
 		int PrevIdx = -1;
 		if(bDoAcquire == true)
 		{
-			InterlockedAdd(RWDataSetIndices[DataSetID], 1, PrevIdx);
+			InterlockedAdd(RWDataSetIndices[DataSetID*5+1], 1, PrevIdx);
 		}
 		return PrevIdx;
 	}
@@ -175,14 +208,19 @@ void ExitStatScope();
 		return InputFloat[RegisterIdx*ComponentBufferSizeRead + InstanceIdx];
 	}
 
+	float InputDataInt(int DataSetIndex, int RegisterIdx, int InstanceIdx)
+	{
+		return InputInt[RegisterIdx*ComponentBufferSizeRead + InstanceIdx];
+	}
+
 	void OutputDataFloat(int DataSetIndex, int RegisterIndex, int InstanceIndex, float Value)
 	{
-		RWOutputFloat[RegisterIndex*ComponentBufferSizeWrite + InstanceIndex + StartInstance] = Value;
+		RWOutputFloat[RegisterIndex*ComponentBufferSizeWrite + InstanceIndex + StartInstancePhase0] = Value;
 	}
 
 	void OutputDataInt(int DataSetIndex, int RegisterIndex, int InstanceIndex, int Value)
 	{
-		RWOutputInt[RegisterIndex*ComponentBufferSizeWrite + InstanceIndex + StartInstance] = Value;
+		RWOutputInt[RegisterIndex*ComponentBufferSizeWrite + InstanceIndex + StartInstancePhase0] = Value;
 	}
 
 	void EnterStatScope(int ID)	{}
@@ -204,36 +242,54 @@ void SimulateMainComputeCS(
 {
 	GDispatchThreadId = DispatchThreadId;
 
-	//if(CopyInstancesBeforeStart == 1)
-	//{
-	//	StartInstance = 0;
-	//}
+	/*
+	if(CopyInstancesBeforeStart == 1)
+	{
+		StartInstancePhase0 = 0;
+	}
+	*/
+
+	uint InstanceID = StartInstancePhase0 + DispatchThreadId.x;
 
-	for(uint InstanceIdx = 0; InstanceIdx<NumInstancesPerThread; InstanceIdx++)
+	if(InstanceID < StartInstancePhase0+TotalNumInstances )
 	{
-		uint InstanceID = StartInstance + InstanceIdx*THREADGROUP_SIZE*NumThreadGroups + DispatchThreadId.x;
+		// calculate event index for the current instance; this is primarily for spawn events, where we do a single read from a secondary DataSet for a subset of particles
+		//  for all-particle events, this will always end up being 0 because NumParticlesPerEvent = total particles in the set
+		uint EventIndex = 0; //(InstanceID-SimulateStartInstance) / NumParticlesPerEvent;
 
-		if(InstanceID < StartInstance+NumInstances )
+		// before the simulation start - either update or copy
+		//
+		if(InstanceID < StartInstancePhase1)
 		{
-			// if we have a start instance >0, then we need to be sure that we don't run on instances before that
-			//
-			if(InstanceID >= SimulateStartInstance)
+			// copy instances before simulation start inst (e.g. for event spawns)
+			/*
+			if(CopyInstancesBeforeStart)
 			{
-				// calculate event index for the current instance; this is primarily for spawn events, where we do a single read from a secondary DataSet for a subset of particles
-				//  for all-particle events, this will always end up being 0 because NumParticlesPerEvent = total particles in the set
-				uint EventIndex = (InstanceID-SimulateStartInstance) / NumParticlesPerEvent;
-
-				SimulateMain(InstanceID, EventIndex);
+				CopyInstance(InstanceID);
+			}
+			else
+			*/
+			{
+				SimulateMain(InstanceID, EventIndex, 1);
 			}
-			// event spawn scripts, for example, will need to copy instances before their start
-			//  TODO: should do this in the VM as well; saves an extra copy
-			//
-			//else if(CopyInstancesBeforeStart)
-			//{
-			//	CopyInstance(InstanceID);
-			//}
 
 		}
+
+		// after simulation start - spawn phase
+		//
+		else if(InstanceID >= StartInstancePhase1)
+		{
+				SimulateMain(InstanceID, EventIndex, 0);		// SimulateMain spawn phase
+		}
+	}
+
+	// fill in the rest of RWDataSetIndices for the DrawIndirect call for the main data set; number of instances is incremented in AcquireIndex
+	if(InstanceID == 0)
+	{
+		RWDataSetIndices[0] = 6;		// IndexCountPerInstance; TODO: need to pass this in from outside based on renderer setup
+		RWDataSetIndices[2] = 0;		// StartIndexLocation
+		RWDataSetIndices[3] = 0;		// BaseVertexLocation
+		RWDataSetIndices[4] = 0;		// StartInstanceLocation
 	}
 }
 
diff --git a/Engine/Shaders/Private/NiagaraMeshVertexFactory.ush b/Engine/Shaders/Private/NiagaraMeshVertexFactory.ush
--- a/Engine/Shaders/Private/NiagaraMeshVertexFactory.ush
+++ b/Engine/Shaders/Private/NiagaraMeshVertexFactory.ush
@@ -5,37 +5,28 @@
 =============================================================================*/
 
 #include "VertexFactoryCommon.ush"
+#include "NiagaraVFParticleAccess.usf"
 
 // These are only enabled for SM5 due to limited interpolators elsewhere
 #define USE_PARTICLE_NODES (FEATURE_LEVEL >= FEATURE_LEVEL_SM5)
 
 // Only support previous transform for deferred feature level
-#define USE_PREVIOUS_TRANSFORM (FEATURE_LEVEL >= FEATURE_LEVEL_SM4)
+//#define USE_PREVIOUS_TRANSFORM (FEATURE_LEVEL >= FEATURE_LEVEL_SM4)
+#define USE_PREVIOUS_TRANSFORM 0
 
 float4x4 PreviousLocalToWorld;
 
+uint MeshFacingMode;
 
 #if USE_PREVIOUS_TRANSFORM
-#if NIAGARA_MESH_INSTANCED
 	Buffer<float4> PrevTransformBuffer;
-#else
-	float4 PrevTransform0;
-	float4 PrevTransform1;
-	float4 PrevTransform2;
-#endif
 #endif
+
+#define MESH_FACING_DEFAULT 0
+#define MESH_FACING_VELOCITY 1
+#define MESH_FACING_CAMERA_POSITION 2
+#define MESH_FACING_CAMERA_PLANE 3
 	
-#if !NIAGARA_MESH_INSTANCED
-	float4 Transform1;
-	float4 Transform2;
-	float4 Transform3;
-	float4 SubUVParams;
-	float SubUVLerp;
-	float4 ParticleDirection;
-	float RelativeTime;
-	float4 DynamicParameter;
-	float4 ParticleColor;
-#endif
 
 struct FVertexFactoryInput
 {
@@ -51,30 +42,8 @@ struct FVertexFactoryInput
 	float2	TexCoords[1] : ATTRIBUTE4;
 #endif
 
-#if NIAGARA_MESH_INSTANCED
-	float4 Transform1 : ATTRIBUTE8;
-	float4 Transform2 : ATTRIBUTE9;
-	float4 Transform3 : ATTRIBUTE10;
-
-	int4 SubUVParams : ATTRIBUTE11;
-
-	/** X: SubUV lerp Y: Particle relative time. */
-	float2 SubUVLerpAndRelTime : ATTRIBUTE12;
-
-	#if USE_DYNAMIC_PARAMETERS
-		float4 DynamicParameter	: ATTRIBUTE13;
-	#endif	//USE_DYNAMIC_PARAMETERS
-
-	float4 ParticleColor : ATTRIBUTE14;
-
-	/** The velocity of the particle, XYZ: direction, W: speed. */
-	float4 ParticleVelocity : ATTRIBUTE15;
-#endif
 
-/** Optional instance ID for vertex layered rendering */
-#if (NIAGARA_MESH_INSTANCED || (ONEPASS_POINTLIGHT_SHADOW && USING_VERTEX_SHADER_LAYER))
 	uint InstanceId	: SV_InstanceID;
-#endif
 };
 
 struct FVertexFactoryInterpolantsVSToPS
@@ -147,6 +116,9 @@ struct FVertexFactoryIntermediates
 	float4 DynamicParameter;
 #endif	//USE_DYNAMIC_PARAMETERS
 
+	/** The position of the particle, XYZ: direction. */
+	float4 ParticlePosition;
+
 	/** The velocity of the particle, XYZ: direction, W: speed. */
 	float4 ParticleVelocity;
 
@@ -156,25 +128,17 @@ struct FVertexFactoryIntermediates
 	/** Relative time. */
 	float RelativeTime;
 
-#if NEEDS_PARTICLE_TRANSFORM
-	float4 LocalToWorld[4];
-#endif
+	float4x4 LocalToWorld;
 };
 
-float4x4 GetParticleTransform(FVertexFactoryInput Input)
+float4x4 GetParticleTransform(FVertexFactoryIntermediates Inters)
 {
-	//Transform is transposed on CPU and position is packed into W. Transpose again to get the correct transform back.
-#if NIAGARA_MESH_INSTANCED
-	return transpose(float4x4(Input.Transform1, Input.Transform2, Input.Transform3, float4(0, 0, 0, 1)));
-#else
-	return transpose(float4x4(Transform1, Transform2, Transform3, float4(0, 0, 0, 1)));
-#endif
+	return Inters.LocalToWorld;
 }
 
 #if USE_PREVIOUS_TRANSFORM
 float4x4 GetParticlePrevTransform(FVertexFactoryInput Input)
 {
-#if NIAGARA_MESH_INSTANCED
 	uint TransformBaseIdx = GetLayerInstanceID(Input.InstanceId) * 3;
 
 	float4x4 PrevTransform4x4 = float4x4(
@@ -183,14 +147,7 @@ float4x4 GetParticlePrevTransform(FVertexFactoryInput Input)
 		PrevTransformBuffer[TransformBaseIdx + 2], 
 		float4(0, 0, 0, 1)
 		);
-#else
-	float4x4 PrevTransform4x4 = float4x4(
-		PrevTransform0, 
-		PrevTransform1, 
-		PrevTransform2, 
-		float4(0, 0, 0, 1)
-		);
-#endif
+
 	//Transform is transposed on CPU and position is packed into W. Transpose again to get the correct transform back.
 	return transpose(PrevTransform4x4);
 }
@@ -287,12 +244,13 @@ FMaterialVertexParameters GetMaterialVertexParameters(FVertexFactoryInput Input,
 	Result.Particle.TranslatedWorldPositionAndSize.w = 1;
 	Result.Particle.Velocity = Intermediates.ParticleVelocity;
 	Result.Particle.RelativeTime = Intermediates.RelativeTime;
+	Result.Particle.Color = Intermediates.ParticleColor;
 
 #if USE_DYNAMIC_PARAMETERS
 	Result.Particle.DynamicParameter = Intermediates.DynamicParameter;
 #endif	//USE_DYNAMIC_PARAMETERS
 
-	float4x4 Transform = GetParticleTransform(Input);
+	float4x4 Transform = GetParticleTransform(Intermediates);
 	Result.InstanceLocalToWorld = Transform;
 #if NEEDS_PARTICLE_TRANSFORM
 	Result.Particle.LocalToWorld[0] = Transform[0];
@@ -326,14 +284,14 @@ half2 GetShadowMapCoordinate(FVertexFactoryInterpolantsVSToPS Interpolants)
 #endif
 
 
-float4 CalcWorldPosition(FVertexFactoryInput Input)
+float4 CalcWorldPosition(FVertexFactoryInput Input, FVertexFactoryIntermediates Intermediates)
 {
-	float4x4 Transform = GetParticleTransform(Input);
+	float4x4 Transform = GetParticleTransform(Intermediates);
 	float3 WorldPosition = mul(Input.Position, Transform).xyz + ResolvedView.PreViewTranslation;
 	return float4(WorldPosition, Input.Position.w);
 }
 
-float4 CalcPrevWorldPosition(FVertexFactoryInput Input)
+float4 CalcPrevWorldPosition(FVertexFactoryInput Input, FVertexFactoryIntermediates Intermediates)
 {
 #if USE_PREVIOUS_TRANSFORM
 	if (MeshParticleVF.PrevTransformAvailable > 0)
@@ -344,41 +302,188 @@ float4 CalcPrevWorldPosition(FVertexFactoryInput Input)
 	}
 	else
 #endif
+	// if we don't have a full previous transform, we can at least calculate the position to get some motion blur
 	{
-		// TODO: this is not the previous position
-		// need to add   - float4(Input.ParticleVelocity.xyz*Input.ParticleVelocity.w, 0.0)  ?
-		return CalcWorldPosition(Input)- float4(Input.ParticleVelocity.xyz*Input.ParticleVelocity.w*MeshParticleVF.DeltaSeconds, 0.0);
+		return CalcWorldPosition(Input, Intermediates)- float4(Intermediates.ParticleVelocity.xyz*Intermediates.ParticleVelocity.w*MeshParticleVF.DeltaSeconds, 0.0);
 	}
 }
 
+// From UnrealMath::FindBetween
+float4 FindQuatBetween(float3 From, float3 To)
+{
+	float NormAB = 1.0f;
+	float w = NormAB + dot(From, To);
+	float4 Result;
+
+	if (w >= 1e-6f * NormAB)
+	{
+		//Axis = cross(From, To);
+		Result = float4(From.y * To.z - From.z * To.y,
+					   From.z * To.x - From.x * To.z,
+					   From.x * To.y - From.y * To.x,
+					   w);
+	}
+	else
+	{
+		// From and To point in opposite directions
+		w = 0.f;
+		Result = abs(From.x) > abs(From.y)
+				? float4(-From.z, 0.f, From.x, w)
+				: float4(0.f, -From.z, From.y, w);
+	}
+
+	Result = normalize(Result);
+	return Result;
+}
+
+float4x4 QuatTo4x4(float4 Q)
+{
+	float x2 = Q.x + Q.x;  float y2 = Q.y + Q.y;  float z2 = Q.z + Q.z;
+	float xx = Q.x * x2;   float xy = Q.x * y2;   float xz = Q.x * z2;
+	float yy = Q.y * y2;   float yz = Q.y * z2;   float zz = Q.z * z2;
+	float wx = Q.w * x2;   float wy = Q.w * y2;   float wz = Q.w * z2;
+
+	float4x4 Mat = {
+	/*M[0][0]*/ 1.0f - (yy + zz),	/*M[1][0]*/ xy - wz,				/*M[2][0]*/ xz + wy,		0.0f,
+	/*M[0][1]*/ xy + wz,			/*M[1][1]*/ 1.0f - (xx + zz),		/*M[2][1]*/ yz - wx,		0.0f,
+	/*M[0][2]*/ xz - wy,			/*M[1][2]*/ yz + wx,				/*M[2][2]*/ 1.0f - (xx + yy),0.0f,
+	0.0f,							0.0f,								0.0f,						1.0f
+	};
+	return Mat;
+}
+
+void AxisAngleToPitchYawRoll(float4 AxisAngle, out float Pitch, out float Yaw, out float Roll)
+{
+	AxisAngle.xyz= normalize(AxisAngle).xyz;
+	Yaw = atan2(AxisAngle.y, AxisAngle.x) * (180.f / PI);
+	if (AxisAngle.y == 0.0f && AxisAngle.x == 0.0f) // atan2 is undefined for 0,0
+	{
+		Yaw = 0.0f;
+	}
+	Pitch = atan2(AxisAngle.z,sqrt(AxisAngle.x*AxisAngle.x+AxisAngle.y*AxisAngle.y)) * (180.f / PI);
+	Roll = AxisAngle.w* (180.f / PI);
+}
+
+
+float4 EulerToQuaternion(float Pitch, float Yaw, float Roll) 
+{
+	float DEG_TO_RAD = PI/(180.f);
+	float DIVIDE_BY_2 = DEG_TO_RAD/2.f;
+	float SP = sin(Pitch*DIVIDE_BY_2);
+	float SY = sin(Yaw*DIVIDE_BY_2);
+	float SR = sin(Roll*DIVIDE_BY_2);
+	float CP = cos(Pitch*DIVIDE_BY_2);
+	float CY = cos(Yaw*DIVIDE_BY_2);
+	float CR = cos(Roll*DIVIDE_BY_2);
+
+	float4 RotationQuat;
+	RotationQuat.x =  CR*SP*SY - SR*CP*CY;
+	RotationQuat.y = -CR*SP*CY - SR*CP*SY;
+	RotationQuat.z =  CR*CP*SY - SR*SP*CY;
+	RotationQuat.w =  CR*CP*CY + SR*SP*SY;
+
+	return RotationQuat;
+}
+
+float4x4 EulerToMat4x4(float Pitch, float Yaw, float Roll)
+{
+	float4 Quat = EulerToQuaternion(Pitch, Yaw, Roll);
+	return QuatTo4x4(Quat);
+}
+
+
 FVertexFactoryIntermediates GetVertexFactoryIntermediates(FVertexFactoryInput Input)
 {
 	FVertexFactoryIntermediates Intermediates = (FVertexFactoryIntermediates)0;
+
+	Intermediates.ParticlePosition = float4(GetVec3(NiagaraMeshVF.PositionDataOffset, Input.InstanceId), 1);
+	Intermediates.ParticleColor = GetVec4(NiagaraMeshVF.ColorDataOffset, Input.InstanceId);
+	Intermediates.ParticleVelocity = float4(GetVec3(NiagaraMeshVF.VelocityDataOffset, Input.InstanceId), 1);
+
+
+#if USE_DYNAMIC_PARAMETERS
+	if (NiagaraMeshVF.MaterialParamDataOffset != -1)
+	{
+		Intermediates.DynamicParameter = GetVec4(NiagaraMeshVF.MaterialParamDataOffset, Input.InstanceId);
+	}
+	else
+	{
+		Intermediates.DynamicParameter = float4(1,1,1,1);
+	}
+#endif
+
+	float4 Transform1 = float4(1.0f, 0.0f, 0.0f, Intermediates.ParticlePosition.x);
+	float4 Transform2 = float4(0.0f, 1.0f, 0.0f, Intermediates.ParticlePosition.y);
+	float4 Transform3 = float4(0.0f, 0.0f, 1.0f, Intermediates.ParticlePosition.z);
+
+	if(NiagaraMeshVF.TransformDataOffset != -1)
+	{
+		float4 AxisAngle = GetVec4(NiagaraMeshVF.TransformDataOffset, Input.InstanceId);
+		
+		//Transform axis to handle local space.
+		//AxisAngle.xyz = mul(float4(AxisAngle.xyz, 0), NiagaraMeshVF.LocalToWorldInverseTransposed);
+
+		float Pitch = 0.0f;
+		float Roll = 0.0f;
+		float Yaw = 0.0f;
+		AxisAngleToPitchYawRoll(AxisAngle, Pitch, Yaw, Roll);
+		float4x4 EulerToMat = EulerToMat4x4(Pitch, Yaw, Roll);
+		Transform1.xyz = EulerToMat[0].xyz;
+		Transform2.xyz = EulerToMat[1].xyz;
+		Transform3.xyz = EulerToMat[2].xyz;
+	}
+	float4x4 ScaleMatrix = float4x4(float4(1.0f, 0.0f, 0.0f, 0.0f), float4(0.0f, 1.0f, 0.0f, 0.0f), float4(0.0f, 0.0f, 1.0f, 0.0f), float4(0.0f, 0.0f, 0.0f, 1.0f));
+	if(NiagaraMeshVF.ScaleDataOffset != -1)
+	{
+		float3 Scale = GetVec3(NiagaraMeshVF.ScaleDataOffset, Input.InstanceId);
+		ScaleMatrix[0][0] *= Scale.x;
+		ScaleMatrix[1][1] *= Scale.y;
+		ScaleMatrix[2][2] *= Scale.z;
+	}	
 	// Swizzle vertex color.
 	Intermediates.VertexColor = Input.VertexColor FCOLOR_COMPONENT_SWIZZLE;
 
 	// World position and localtoworld transform.
-#if NIAGARA_MESH_INSTANCED
-	float3 ParticleWorldPosition = float3( Input.Transform1.w, Input.Transform2.w, Input.Transform3.w);
-	float4x4 LocalToWorldMat = float4x4(Input.Transform1, Input.Transform2, Input.Transform3, float4(0.0f, 0.0f, 0.0f, 1.0f));
-#if NEEDS_PARTICLE_TRANSFORM
-	Intermediates.LocalToWorld[0] = LocalToWorldMat[0];
-	Intermediates.LocalToWorld[1] = LocalToWorldMat[1];
-	Intermediates.LocalToWorld[2] = LocalToWorldMat[2];
-	Intermediates.LocalToWorld[3] = LocalToWorldMat[3];
-#endif
-#else
-	float3 ParticleWorldPosition = float3( Transform1.w, Transform2.w, Transform3.w);
-#endif
+	float4x4 LocalToWorldMat = float4x4(Transform1, Transform2, Transform3, float4(0.0f, 0.0f, 0.0f, 1.0f));
+	
+	Intermediates.ParticleTranslatedWorldPosition = Intermediates.ParticlePosition + ResolvedView.PreViewTranslation.xyz;
 
-	Intermediates.ParticleTranslatedWorldPosition = ParticleWorldPosition + ResolvedView.PreViewTranslation.xyz;
+	//LocalToWorldMat = mul(ScaleMatrix, LocalToWorldMat);
 
-#if NIAGARA_MESH_INSTANCED
-	Intermediates.ParticleColor = Input.ParticleColor;
-#else
-	Intermediates.ParticleColor = ParticleColor;
-#endif
+	if (MeshFacingMode == MESH_FACING_VELOCITY)
+	{
+		float4 ToAlignmentQuat = FindQuatBetween(float3(1.0f, 0.0f, 0.0f), normalize(Intermediates.ParticleVelocity.xyz));
+		float4x4 ToAlignmentMat = QuatTo4x4(ToAlignmentQuat);
+		Intermediates.LocalToWorld =  mul(LocalToWorldMat, mul(ToAlignmentMat, ScaleMatrix));
+	}
+	else if (MeshFacingMode == MESH_FACING_CAMERA_POSITION)
+	{
+		float3 WorldPosition = Intermediates.ParticleTranslatedWorldPosition;
+		float3 CameraVector = ResolvedView.WorldCameraOrigin.xyz - WorldPosition.xyz;
+		float3 CameraDirection = normalize(CameraVector);
+		float4 ToAlignmentQuat = FindQuatBetween(float3(1.0f, 0.0f, 0.0f), CameraDirection);
+		float4x4 ToAlignmentMat = QuatTo4x4(ToAlignmentQuat);
+		Intermediates.LocalToWorld = mul(LocalToWorldMat, mul(ToAlignmentMat, ScaleMatrix));
+	}
+	else if (MeshFacingMode == MESH_FACING_CAMERA_PLANE)
+	{
+		float3 CameraFacingVector = -ResolvedView.ViewForward;
+		float4 ToAlignmentQuat = FindQuatBetween(float3(1.0f, 0.0f, 0.0f), CameraFacingVector);
+		float4x4 ToAlignmentMat = QuatTo4x4(ToAlignmentQuat);
+		Intermediates.LocalToWorld =  mul(LocalToWorldMat, mul(ToAlignmentMat, ScaleMatrix));
+	}
+	else
+	{
+		Intermediates.LocalToWorld =  mul(LocalToWorldMat, ScaleMatrix); 
+	}
 
+	//Handle local space. TODO: Optimize behind define.
+	//Apply local space to local to world, also pretranspse here to avoid it in the vs&ps
+	Intermediates.LocalToWorld = transpose(Intermediates.LocalToWorld);
+	Intermediates.LocalToWorld = mul(Intermediates.LocalToWorld, NiagaraMeshVF.LocalToWorld);
+	Intermediates.ParticlePosition = mul(Intermediates.ParticlePosition, NiagaraMeshVF.LocalToWorld);
+	Intermediates.ParticleVelocity.xyz = mul(float4(Intermediates.ParticleVelocity.xyz,0), NiagaraMeshVF.LocalToWorldInverseTransposed);
+	
 #if NUM_MATERIAL_TEXCOORDS_VERTEX
 	for(int CoordinateIndex = 0; CoordinateIndex < NUM_MATERIAL_TEXCOORDS_VERTEX; CoordinateIndex++)
 	{
@@ -387,34 +492,12 @@ FVertexFactoryIntermediates GetVertexFactoryIntermediates(FVertexFactoryInput In
 #endif
 
 #if USE_PARTICLE_SUBUVS
-	#if NIAGARA_MESH_INSTANCED
 		Intermediates.SubUVCoords.xy = (float2( Input.SubUVParams.x, Input.SubUVParams.y ) + Input.TexCoords[0].xy) * MeshParticleVF.SubImageSize.xy;
 		Intermediates.SubUVCoords.zw = (float2( Input.SubUVParams.z, Input.SubUVParams.w ) + Input.TexCoords[0].xy) * MeshParticleVF.SubImageSize.xy;
-		Intermediates.SubUVLerp = Input.SubUVLerpAndRelTime.x;
-	#else
-		Intermediates.SubUVCoords.xy = (float2( SubUVParams.x, SubUVParams.y ) + Input.TexCoords[0].xy) * MeshParticleVF.SubImageSize.xy;
-		Intermediates.SubUVCoords.zw = (float2( SubUVParams.z, SubUVParams.w ) + Input.TexCoords[0].xy) * MeshParticleVF.SubImageSize.xy;
-		Intermediates.SubUVLerp = SubUVLerp.x;
-	#endif
+		//Intermediates.SubUVLerp = Input.SubUVLerpAndRelTime.x;
 #endif 
 
-#if USE_DYNAMIC_PARAMETERS
-	#if NIAGARA_MESH_INSTANCED
-		Intermediates.DynamicParameter = Input.DynamicParameter;
-	#else
-		Intermediates.DynamicParameter = DynamicParameter;
-	#endif
-#endif	//USE_DYNAMIC_PARAMETERS
-
-	Intermediates.ParticleTranslatedWorldPosition = ParticleWorldPosition.xyz + ResolvedView.PreViewTranslation.xyz;
-
-#if NIAGARA_MESH_INSTANCED
-	Intermediates.ParticleVelocity = Input.ParticleVelocity;
-	Intermediates.RelativeTime = Input.SubUVLerpAndRelTime.y;
-#else
-	Intermediates.ParticleVelocity = ParticleDirection;
-	Intermediates.RelativeTime = RelativeTime;
-#endif
+	//Intermediates.RelativeTime = Input.SubUVLerpAndRelTime.y;
 
 	return Intermediates;
 }
@@ -451,7 +534,7 @@ float3x3 VertexFactoryGetTangentToLocal( FVertexFactoryInput Input, FVertexFacto
 // @return translated world position
 float4 VertexFactoryGetWorldPosition(FVertexFactoryInput Input, FVertexFactoryIntermediates Intermediates)
 {
-	return CalcWorldPosition(Input);
+	return CalcWorldPosition(Input, Intermediates);
 }
 
 float4 VertexFactoryGetRasterizedWorldPosition(FVertexFactoryInput Input, FVertexFactoryIntermediates Intermediates, float4 InWorldPosition)
@@ -468,7 +551,7 @@ void CalcTangentToWorld(FVertexFactoryInput Input, FVertexFactoryIntermediates I
 {
 	float3x3 TangentToLocal = VertexFactoryGetTangentToLocal(Input, Intermediates);
 
-	float4x4 Transform = GetParticleTransform(Input);
+	float4x4 Transform = GetParticleTransform(Intermediates);
 	float3x3 TangentToWorld = mul(TangentToLocal, (float3x3)Transform);
 
 	// Normalize to remove scaling. Incorrect but faster than computing the inverse-transpose.
@@ -549,7 +632,7 @@ FVertexFactoryInterpolantsVSToPS VertexFactoryGetInterpolantsVSToPS(FVertexFacto
 // @return previous translated world position
 float4 VertexFactoryGetPreviousWorldPosition(FVertexFactoryInput Input, FVertexFactoryIntermediates Intermediates)
 {	
-	return CalcPrevWorldPosition(Input);
+	return CalcPrevWorldPosition(Input, Intermediates);
 }
 
 #if USING_TESSELLATION
@@ -651,16 +734,10 @@ float4 VertexFactoryGetPreviousWorldPosition(FVertexFactoryInput Input, FVertexF
 #if INSTANCED_STEREO
 
 // Instanced stereo sets the eye index explicitly for instanced geometry
-#if NIAGARA_MESH_INSTANCED
 uint InstancedEyeIndex;
-#endif
 
 uint VertexFactoryGetEyeIndex(uint InstanceId)
 {
-	#if NIAGARA_MESH_INSTANCED
-		return InstancedEyeIndex & bIsInstancedStereo;
-	#else
-		return InstanceId & 1;
-	#endif
+	return InstancedEyeIndex & bIsInstancedStereo;
 }
 #endif
diff --git a/Engine/Shaders/Private/NiagaraRibbonVertexFactory.ush b/Engine/Shaders/Private/NiagaraRibbonVertexFactory.ush
new file mode 100644
--- /dev/null
+++ b/Engine/Shaders/Private/NiagaraRibbonVertexFactory.ush
@@ -0,0 +1,361 @@
+// Copyright 1998-2017 Epic Games, Inc. All Rights Reserved.
+
+/*=============================================================================
+	ParticleBeamTrailVertexFactory.hlsl: Particle vertex factory shader code.
+=============================================================================*/
+
+#include "VertexFactoryCommon.ush"
+
+
+
+
+float3 TransformPosition(float3 InPosition)
+{
+	return mul(float4(InPosition,1.0), NiagaraSpriteVF.LocalToWorld).xyz;
+}
+
+float3 TransformVector(float3 InVector)
+{
+	return mul(float4(InVector,0.0), NiagaraSpriteVF.LocalToWorldInverseTransposed).xyz;
+}
+
+
+/*********************************************************************************
+ *							Particle specific
+ *********************************************************************************/
+/*
+float3 GetNiagaraParticlePosition(uint InstanceID)
+{
+	return TransformPosition(GetVec3(NiagaraRibbonVF.PositionDataOffset, InstanceID));
+}
+float GetNiagaraParticleRibbonWidth(uint InstanceID)
+{
+	return GetFloat(NiagaraRibbonVF.WidthDataOffset, InstanceID);
+}
+float GetNiagaraParticleRibbonTwist(uint InstanceID)
+{
+	return GetFloat(NiagaraRibbonVF.TwistDataOffset, InstanceID);
+}
+float4 GetNiagaraParticleColor(uint InstanceID)
+{
+	return GetVec4(NiagaraRibbonVF.ColorDataOffset, InstanceID);
+}
+#if USE_DYNAMIC_PARAMETERS
+float4 GetNiagaraParticleDynamicParameters(uint InstanceID)
+{
+	if (NiagaraRibbonVF.MaterialParamDataOffset == -1)
+	{
+		return float4(1.0f, 1.0f, 1.0f, 1.0f);
+	}
+	return GetVec4(NiagaraRibbonVF.MaterialParamDataOffset, InstanceID);
+}
+#endif
+
+
+float3 GetNiagaraParticleFacingVector(uint InstanceID)
+{
+	float3 Val = float3(0.0f, 0.0f, 0.0f);
+    if(NiagaraRibbonVF.FacingOffset != -1)
+	{
+		Val = GetVec3(NiagaraRibbonVF.FacingOffset, InstanceID);
+	}
+	return TransformVector(Val);
+}
+*/
+
+
+
+
+struct FVertexFactoryInput
+{
+	float3	Position				: ATTRIBUTE0;
+	float3 Direction				: ATTRIBUTE1;
+	float	Size					: ATTRIBUTE2;
+	float4	Color					: ATTRIBUTE3;
+	float4	TexCoord				: ATTRIBUTE4;
+
+	float Twist						: ATTRIBUTE5;
+	float3 CustomFacingVector		: ATTRIBUTE6;
+
+#if USE_DYNAMIC_PARAMETERS
+	float4 DynamicParameter0		: ATTRIBUTE7;
+#endif	//USE_DYNAMIC_PARAMETERS
+
+/** Optional instance ID for vertex layered rendering */
+#if FEATURE_LEVEL >= FEATURE_LEVEL_SM4
+	uint VertexID	: SV_VertexID;
+#endif
+};
+
+struct FVertexFactoryInterpolantsVSToPS
+{
+	// First row of the tangent to world matrix
+	float3	TangentToWorld0	: TANGENTTOWORLD0;
+	// Last row of the tangent to world matrix in xyz
+	float4	TangentToWorld2	: TANGENTTOWORLD2;
+
+	float4	TexCoord			: TEXCOORD2;
+
+#if NEEDS_PARTICLE_COLOR
+	float4  Color				: TEXCOORD3;
+#endif
+
+#if USE_DYNAMIC_PARAMETERS
+	float4 DynamicParameter0	: TEXCOORD4;
+#endif	//USE_DYNAMIC_PARAMETERS
+
+#if INSTANCED_STEREO
+	nointerpolation uint PackedEyeIndex : PACKED_EYE_INDEX;
+#endif
+};
+
+struct FVertexFactoryIntermediates
+{
+	/** The color of the sprite. */
+	float4 Color;
+	float RibbonWidth;
+	float4 Right, Up;
+	float3 CameraDirection;
+};
+
+/** Converts from vertex factory specific interpolants to a FMaterialPixelParameters, which is used by material inputs. */
+FMaterialPixelParameters GetMaterialPixelParameters(FVertexFactoryInterpolantsVSToPS Interpolants, float4 SvPosition)
+{
+	// GetMaterialPixelParameters is responsible for fully initializing the result
+	FMaterialPixelParameters Result = MakeInitializedMaterialPixelParameters();
+
+#if NUM_MATERIAL_TEXCOORDS
+	#if NUM_MATERIAL_TEXCOORDS >= 1
+		Result.TexCoords[0] = Interpolants.TexCoord.xy;
+	#if NUM_MATERIAL_TEXCOORDS >= 2
+		Result.TexCoords[1] = Interpolants.TexCoord.zw;
+	#endif	// >= 2
+	#endif	// >= 1
+	#if NUM_MATERIAL_TEXCOORDS > 2
+		UNROLL
+		for(int CoordinateIndex = 2;CoordinateIndex < NUM_MATERIAL_TEXCOORDS;CoordinateIndex++)
+		{
+			Result.TexCoords[CoordinateIndex] = Interpolants.TexCoord.xy;
+		}
+	#endif	// > 2
+#endif
+	Result.VertexColor = 1;
+
+#if NEEDS_PARTICLE_COLOR
+	Result.Particle.Color = Interpolants.Color;
+#endif
+
+	half3 TangentToWorld0 = Interpolants.TangentToWorld0.xyz;
+	half4 TangentToWorld2 = Interpolants.TangentToWorld2;
+	Result.TangentToWorld = AssembleTangentToWorld( TangentToWorld0, TangentToWorld2 );
+	Result.UnMirrored = 1;
+
+#if LIGHTMAP_UV_ACCESS
+	Result.LightmapUVs = Interpolants.TexCoord.xy;
+#endif
+
+	Result.TwoSidedSign = 1;
+
+#if USE_DYNAMIC_PARAMETERS
+	Result.Particle.DynamicParameter = Interpolants.DynamicParameter0;
+#endif	//USE_DYNAMIC_PARAMETERS
+
+	return Result;
+}
+
+/** Converts from vertex factory specific input to a FMaterialVertexParameters, which is used by vertex shader material inputs. */
+FMaterialVertexParameters GetMaterialVertexParameters(FVertexFactoryInput Input, FVertexFactoryIntermediates Intermediates, float3 WorldPosition, float3x3 TangentToLocal)
+{
+	FMaterialVertexParameters Result = (FMaterialVertexParameters)0;
+	Result.WorldPosition = WorldPosition;
+	Result.VertexColor = Input.Color;
+	Result.TangentToWorld = mul(TangentToLocal, GetLocalToWorld3x3()); 
+	Result.Particle.Color = Intermediates.Color;
+	Result.PreSkinnedPosition = Input.Position.xyz;
+	Result.PreSkinnedNormal = TangentToLocal[2].xyz;
+
+#if NUM_MATERIAL_TEXCOORDS_VERTEX
+	for(int CoordinateIndex = 0; CoordinateIndex < NUM_MATERIAL_TEXCOORDS_VERTEX; CoordinateIndex++)
+	{
+		Result.TexCoords[CoordinateIndex] = Input.TexCoord.xy;
+	}
+#endif
+	return Result;
+}
+
+void GetLightMapCoordinates(FVertexFactoryInterpolantsVSToPS Interpolants, out float2 LightmapUV0, out float2 LightmapUV1)
+{
+	LightmapUV0 = LightmapUV1 = 0;
+}
+
+half2 GetShadowMapCoordinate(FVertexFactoryInterpolantsVSToPS Interpolants)
+{
+	return 0;
+}
+
+float3 SafeNormalize(float3 V)
+{
+	return V / sqrt(max(dot(V,V),0.01));
+}
+
+void GetTangents(FVertexFactoryInput Input,out float4 Right,out float4 Up, out float3 CameraDirection)
+{
+	float4	Position			= TransformLocalToTranslatedWorld(Input.Position.xyz);
+			//OldPosition			= TransformLocalToTranslatedWorld(Input.OldPosition.xyz);
+
+	CameraDirection		= SafeNormalize(ResolvedView.WorldCameraOrigin - Input.Position.xyz );		// TODO: handle local space
+	float3 ParticleDirection	= Input.Direction; //SafeNormalize(Position.xyz - OldPosition.xyz);
+
+	float4	Right_Square		= BeamTrailVF.CameraRight,
+			Up_Square			= BeamTrailVF.CameraUp;			
+			
+	float4	Right_Rotated		= (-1.0 * cos(Input.Twist) * Up_Square) + (sin(Input.Twist) * Right_Square),
+			Up_Rotated			= (       sin(Input.Twist) * Up_Square) + (cos(Input.Twist) * Right_Square);
+
+	float4	Right_Velocity		= float4( SafeNormalize( cross( CameraDirection, ParticleDirection	) ), 0.0 ),
+			Up_Velocity			= float4( ParticleDirection, 0.0 );
+
+	float3 NormDir = SafeNormalize(Input.Direction.xyz);
+	float3 RightDir;
+	if(NiagaraRibbonVF.UseCustomFacing)
+	{
+		RightDir = SafeNormalize(cross(NormDir, Input.CustomFacingVector));
+	}
+	else
+	{
+		//float Blend = abs(dot(NormDir, SafeNormalize(CameraDirection.xyz)));
+		RightDir = SafeNormalize(cross(NormDir, CameraDirection.xyz));
+	}
+	
+	Right = float4(RightDir, 0.0f);
+	Up = float4(NormDir, 0.0f);
+	return;
+	/*
+	//	enum EParticleScreenAlignment
+	//	{
+	//		PSA_Square,
+	//		PSA_Rectangle,
+	//		PSA_Velocity
+	//	};
+	Right				= BeamTrailVF.ScreenAlignment.x > 1.5f ? Right_Velocity : Right_Rotated;
+	Up					= BeamTrailVF.ScreenAlignment.x > 1.5f ? Up_Velocity	: Up_Rotated;
+	*/
+}
+
+float4 CalcWorldPosition(FVertexFactoryInput Input, FVertexFactoryIntermediates Intermediates)
+{
+	// expand and align
+	int Sgn = (((int)Input.VertexID & 0x1) * 2) - 1;
+	float3 NormDir = SafeNormalize(Input.Direction.xyz);
+	float3 RightDir;
+	if(NiagaraRibbonVF.UseCustomFacing)
+	{
+		RightDir = SafeNormalize(cross(NormDir, Input.CustomFacingVector));
+	}
+	else
+	{
+		//float Blend = abs(dot(NormDir, SafeNormalize(Intermediates.CameraDirection.xyz)));
+		RightDir = SafeNormalize(cross(NormDir, Intermediates.CameraDirection.xyz));
+	}
+	float3 UpDir = cross(RightDir, NormDir);
+
+	//float4	Right_Rotated		= (-1.0 * cos(Input.Twist) * UpDir) + (sin(Input.Twist) * RightDir),
+	//		Up_Rotated			= (       sin(Input.Twist) * UpDir) + (cos(Input.Twist) * RightDir);
+
+	float3 Pos = Input.Position.xyz + RightDir * Sgn * Input.Size.xxx;
+	return TransformLocalToTranslatedWorld(Pos);
+}
+
+/** derive basis vectors */
+float3x3 CalcTangentBasis(FVertexFactoryInput Input)
+{
+	float4	Right,
+			Up;
+	float3	CameraDir;
+	GetTangents(Input,Right,Up, CameraDir);
+	return float3x3(
+			Right.xyz,
+			Up.xyz,
+			-normalize(cross(Right.xyz,Up.xyz))
+			);
+}
+
+FVertexFactoryIntermediates GetVertexFactoryIntermediates(FVertexFactoryInput Input)
+{
+	FVertexFactoryIntermediates Intermediates;
+
+	Intermediates.Color = Input.Color;
+    Intermediates.RibbonWidth = Input.Size;
+	GetTangents(Input, Intermediates.Right, Intermediates.Up, Intermediates.CameraDirection);
+
+	return Intermediates;
+}
+
+float4 VertexFactoryGetWorldPosition(FVertexFactoryInput Input, FVertexFactoryIntermediates Intermediates)
+{
+	return CalcWorldPosition(Input, Intermediates);
+}
+
+float4 VertexFactoryGetRasterizedWorldPosition(FVertexFactoryInput Input, FVertexFactoryIntermediates Intermediates, float4 InWorldPosition)
+{
+	return InWorldPosition;
+}
+
+float3 VertexFactoryGetPositionForVertexLighting(FVertexFactoryInput Input, FVertexFactoryIntermediates Intermediates, float3 TranslatedWorldPosition)
+{
+	return TranslatedWorldPosition;
+}
+
+FVertexFactoryInterpolantsVSToPS VertexFactoryGetInterpolantsVSToPS(FVertexFactoryInput Input, FVertexFactoryIntermediates Intermediates, FMaterialVertexParameters VertexParameters)
+{
+	FVertexFactoryInterpolantsVSToPS Interpolants;
+	Interpolants.TexCoord	= Input.TexCoord;
+
+#if NEEDS_PARTICLE_COLOR
+	Interpolants.Color		= Intermediates.Color;
+#endif
+
+#if USE_DYNAMIC_PARAMETERS
+	Interpolants.DynamicParameter0 = Input.DynamicParameter0;
+#endif	//USE_DYNAMIC_PARAMETERS
+
+	float3x3 TangentToLocal = CalcTangentBasis(Input);
+	float3x3 TangentToWorld = mul(TangentToLocal, GetLocalToWorld3x3());
+
+	Interpolants.TangentToWorld0 = TangentToWorld[0];
+	Interpolants.TangentToWorld2 = float4(TangentToWorld[2], sign(determinant(TangentToLocal)));
+
+#if INSTANCED_STEREO
+	Interpolants.PackedEyeIndex = 0;
+#endif
+
+	return Interpolants;
+}
+
+float4 VertexFactoryGetPreviousWorldPosition(FVertexFactoryInput Input, FVertexFactoryIntermediates Intermediates)
+{
+	return VertexFactoryGetWorldPosition(Input, Intermediates);
+}
+
+/**
+* Get the 3x3 tangent basis vectors for this vertex factory
+*
+* @param Input - vertex input stream structure
+* @return 3x3 matrix
+*/
+float3x3 VertexFactoryGetTangentToLocal( FVertexFactoryInput Input, FVertexFactoryIntermediates Intermediates )
+{
+	return CalcTangentBasis(Input);
+}
+
+#if INSTANCED_STEREO
+uint VertexFactoryGetEyeIndex(uint InstanceId)
+{
+	return InstanceId & 1;
+}
+#endif
+
+float4 VertexFactoryGetTranslatedPrimitiveVolumeBounds(FVertexFactoryInterpolantsVSToPS Interpolants)
+{
+	return 0;
+}
\ No newline at end of file
diff --git a/Engine/Shaders/Private/NiagaraSpriteVertexFactory.ush b/Engine/Shaders/Private/NiagaraSpriteVertexFactory.ush
--- a/Engine/Shaders/Private/NiagaraSpriteVertexFactory.ush
+++ b/Engine/Shaders/Private/NiagaraSpriteVertexFactory.ush
@@ -15,11 +15,29 @@
 #define USE_PARTICLE_VELOCITY (FEATURE_LEVEL >= FEATURE_LEVEL_SM5 && NEEDS_PARTICLE_VELOCITY)
 #define USE_PARTICLE_TIME (FEATURE_LEVEL >= FEATURE_LEVEL_SM5 && NEEDS_PARTICLE_TIME)
 #define USE_PARTICLE_SIZE (FEATURE_LEVEL >= FEATURE_LEVEL_SM5 && NEEDS_PARTICLE_SIZE)
-
-bool UseCustomAlignment;
-bool UseVectorAlignment;
-bool CameraPlaneFacing;
-
+#define USE_PARTICLE_RANDOM (FEATURE_LEVEL >= FEATURE_LEVEL_SM5 && NEEDS_PARTICLE_RANDOM)
+
+uint ParticleAlignmentMode;
+uint ParticleFacingMode;
+
+// Must keep these in sync with ENiagaraSpriteAlignment in NiagaraSpriteRendererProperties.h
+/** Only Particles.SpriteRotation and FacingMode impact the alignment of the particle.*/
+#define ALIGNMENT_MODE_UNALIGNED 0
+/** Imagine the particle texture having an arrow pointing up, this mode makes the arrow point in the direction of the Particles.Velocity attribute. Particles.SpriteRotation rotation is ignored.*/
+#define ALIGNMENT_MODE_VELOCITYALIGNED 1
+/** Imagine the particle texture having an arrow pointing up, this mode makes the arrow point towards the axis defined by the Particles.SpriteAlignment attribute, after taking Particles.SpriteRotation and the FacingMode into account.*/
+#define ALIGNMENT_MODE_CUSTOMALIGNMENT 2
+
+/** The sprite billboard origin is always "looking at" the camera origin, trying to keep its up axis aligned to the camera's up axis. */
+#define FACING_MODE_FACECAMERA 0
+/** The sprite billboard plane is completely parallel to the camera plane. Particle always looks "flat" */
+#define FACING_MODE_FACECAMERAPLANE 1
+/** The sprite billboard faces toward the "Particles.SpriteFacing" vector attribute, using the per-axis CustomFacingVectorMask as a lerp factor from the standard FaceCamera mode.*/
+#define FACING_MODE_CUSTOMFACINGVECTOR 2
+/** Faces the camera position, but is not dependent on the camera rotation.  This method produces more stable particles under camera rotation. Uses the up axis of (0,0,1).*/
+#define FACING_MODE_FACECAMERAPOSITION 3 
+/** Blends between FaceCamera and FaceCameraPosition.*/
+#define FACING_MODE_FACECAMERADISTANCEBLEND 4
 
 struct FVertexFactoryInput
 {
@@ -90,6 +108,10 @@ struct FVertexFactoryInterpolantsVSToPS
 	float2 ParticleSize : PARTICLE_SIZE;
 #endif
 
+#if USE_PARTICLE_RANDOM
+	float ParticleRandom : PARTICLE_RANDOM;
+#endif
+
 #if INSTANCED_STEREO
 	nointerpolation uint PackedEyeIndex : PACKED_EYE_INDEX;
 #endif
@@ -132,6 +154,8 @@ struct FVertexFactoryIntermediates
 	float3x3 TangentToLocal;
 	/** Size of the particle */
 	float2 ParticleSize;
+	/** Particle Random value */
+	float ParticleRandom;
 };
 
 #if NUM_MATERIAL_TEXCOORDS
@@ -239,6 +263,12 @@ FMaterialPixelParameters GetMaterialPixelParameters(FVertexFactoryInterpolantsVS
 	Result.Particle.Size = Interpolants.ParticleSize;
 #endif
 
+#if USE_PARTICLE_RANDOM
+	Result.Particle.Random = Interpolants.ParticleRandom;
+#else
+	Result.Particle.Random = 0.0f;
+#endif
+
 	Result.TwoSidedSign = 1;
 	return Result;
 }
@@ -272,6 +302,12 @@ FMaterialVertexParameters GetMaterialVertexParameters(FVertexFactoryInput Input,
 	Result.Particle.RelativeTime = Intermediates.RelativeTime;
 #endif
 
+#if USE_PARTICLE_RANDOM
+	Result.Particle.Random = Intermediates.ParticleRandom;
+#else
+	Result.Particle.Random = 0.0f;
+#endif
+
 #if NUM_MATERIAL_TEXCOORDS_VERTEX
 	#if NUM_MATERIAL_TEXCOORDS_VERTEX >= 1
 		Result.TexCoords[0] = Intermediates.TexCoord.xy;
@@ -319,37 +355,72 @@ void GetTangents(FVertexFactoryInput Input, float3 WorldPosition, float3 OldWorl
 	float3 ResolvedViewRight = lerp(ResolvedView.ViewRight, ResolvedView.HMDViewNoRollRight, NiagaraSpriteVF.RemoveHMDRoll);
 	float3 ResolvedViewUp = lerp(ResolvedView.ViewUp, ResolvedView.HMDViewNoRollUp, NiagaraSpriteVF.RemoveHMDRoll);
 
-	// camera related vectors
+	// The standard camera facing vectors
 	float3 CameraRight = ResolvedViewRight;
 	float3 CameraUp = -ResolvedViewUp;
+	float3 CameraFacingVector = -ResolvedView.ViewForward;
+
+	// Determine the vector from the particle to the camera and the particle's movement direction.
 	float3 CameraVector = ResolvedView.WorldCameraOrigin.xyz - WorldPosition.xyz;
 	float3 CameraDirection = SafeNormalize(CameraVector);
-	float3 CameraFacingVector = Input.CustomFacingVector * NiagaraSpriteVF.CustomFacingVectorMask + CameraDirection * (float3(1.0,1.0,1.0)-NiagaraSpriteVF.CustomFacingVectorMask);
-	CameraFacingVector = lerp(CameraFacingVector, -ResolvedView.ViewForward, CameraPlaneFacing);
+	float3 ParticleDirection	= SafeNormalize(WorldPosition.xyz - OldWorldPosition.xyz);
+
+	// First interpolate between the base facing mode and the custom facing mode. We assume that the code setting up this call has the facing vector mask
+	// set to 0,0,0,0 if not in FACING_MODE_CUSTOMFACINGVECTOR.
+	float3 CustomFacingVectorMask =  NiagaraSpriteVF.CustomFacingVectorMask;
+	float3 CustomFacingVector = Input.CustomFacingVector * CustomFacingVectorMask + CameraDirection * (float3(1.0,1.0,1.0)-CustomFacingVectorMask);
 	
+	float WeightDefaultCameraMode = (ParticleAlignmentMode == ALIGNMENT_MODE_UNALIGNED) * (ParticleFacingMode == FACING_MODE_FACECAMERA);
+ 	float WeightPlaneFacing =  (ParticleAlignmentMode == ALIGNMENT_MODE_UNALIGNED) * (ParticleFacingMode == FACING_MODE_FACECAMERAPLANE);
+	float WeightCustomFacing =  (ParticleAlignmentMode == ALIGNMENT_MODE_UNALIGNED) * (ParticleFacingMode == FACING_MODE_CUSTOMFACINGVECTOR);
+	float WeightPositionFacing =  	 (ParticleAlignmentMode == ALIGNMENT_MODE_UNALIGNED) * (ParticleFacingMode == FACING_MODE_FACECAMERAPOSITION);
+	float WeightVelocityAligned = ParticleAlignmentMode != ALIGNMENT_MODE_UNALIGNED;
+	
+	// Adjust weights if you are in the blend mode based off of camera distance.
 	BRANCH
-	if (NiagaraSpriteVF.CameraFacingBlend.x > 0.f)
+	if (NiagaraSpriteVF.CameraFacingBlend.x > 0.f && ParticleFacingMode == FACING_MODE_FACECAMERADISTANCEBLEND)
 	{
 		float CameraDistanceSq = dot(CameraVector, CameraVector);
 		float AlignmentMode = saturate(CameraDistanceSq * NiagaraSpriteVF.CameraFacingBlend.y - NiagaraSpriteVF.CameraFacingBlend.z);
-		CameraFacingVector = -ResolvedView.ViewForward * CameraFacingVector * (1.0-AlignmentMode);
+		WeightDefaultCameraMode = 1.f - AlignmentMode;
+		WeightPositionFacing = AlignmentMode;
+		WeightVelocityAligned = 0.0f;
 	}
-
+	
 	// the particle's movement direction.
-	float3 ParticleDirection	= SafeNormalize(WorldPosition.xyz - OldWorldPosition.xyz);
-	float3 DirectionVector	= lerp(ParticleDirection, Input.CustomAlignmentVector, UseCustomAlignment);
-
-	// Tangent vectors for PSA_Velocity.
-	float3 Right_Aligned	= SafeNormalize(cross(CameraDirection, DirectionVector));
-	float3 Up_Aligned		= -DirectionVector;
-
-	// Tangent vectors for camera facing position (can be based on camera position or the custom facing vector)
-	//
-	float3 Right_FacingCameraPosition = SafeNormalize(cross(CameraFacingVector,ResolvedViewUp));
-	float3 Up_FacingCameraPosition = cross(CameraFacingVector, Right_FacingCameraPosition);
-
-	float3 Right = SafeNormalize( lerp(Right_FacingCameraPosition, Right_Aligned, UseVectorAlignment) );
-	float3 Up = SafeNormalize( cross(CameraFacingVector, Right) );
+	float3 DirectionVector	= lerp(ParticleDirection, Input.CustomAlignmentVector, ParticleAlignmentMode == ALIGNMENT_MODE_CUSTOMALIGNMENT);
+	
+	// Tangent vectors for any non-default alignment mode.
+	float3 CameraRightAligned	= SafeNormalize(cross(CustomFacingVector, DirectionVector));
+	float3 CameraUpAligned = lerp(-DirectionVector, cross(CustomFacingVector, CameraRightAligned), ParticleFacingMode == FACING_MODE_CUSTOMFACINGVECTOR);
+	//float3 CameraUpAligned = -DirectionVector;
+
+	// Tangent vectors for camera plane facing mode
+	float3 CameraRightPlaneFacing = SafeNormalize(cross(CameraFacingVector,ResolvedViewUp));
+	float3 CameraUpPlaneFacing = cross(CameraFacingVector, CameraRightPlaneFacing);
+
+	// Tangent vectors for camera custom facing mode
+	float3 CameraRightCustomFacing = SafeNormalize(cross(CustomFacingVector,ResolvedViewUp));
+	float3 CameraUpCustomFacing = cross(CustomFacingVector, CameraRightCustomFacing);
+
+	// Tangent vectors for camera facing position.
+	float3 CameraRightPosition = SafeNormalize(cross(CameraDirection,float3(0,0,1)));
+	float3 CameraUpPosition = cross(CameraDirection, CameraRightPosition);
+
+	float3 RightCombined =	WeightDefaultCameraMode * CameraRight.xyz + 
+							WeightCustomFacing * CameraRightCustomFacing.xyz +
+							WeightVelocityAligned * CameraRightAligned.xyz +
+							WeightPlaneFacing * CameraRightPlaneFacing.xyz + 
+							WeightPositionFacing * CameraRightPosition.xyz;
+
+	float3 UpCombined = WeightDefaultCameraMode * CameraUp.xyz + 
+						WeightCustomFacing * CameraUpCustomFacing.xyz +
+						WeightVelocityAligned * CameraUpAligned.xyz +
+						WeightPlaneFacing * CameraUpPlaneFacing.xyz +
+						WeightPositionFacing * CameraUpPosition.xyz;
+
+	float3 Right = SafeNormalize( RightCombined );
+	float3 Up = SafeNormalize( UpCombined );
 
 	// Determine the angle of rotation.
 	float SinRotation;
@@ -440,7 +511,15 @@ void ComputeBillboardUVs(FVertexFactoryInput Input, out float2 UVForPosition, ou
 	}
 }
 
+float3 TransformPosition(float3 InPosition)
+{
+	return mul(float4(InPosition,1.0), NiagaraSpriteVF.LocalToWorld).xyz;
+}
 
+float3 TransformVector(float3 InVector)
+{
+	return mul(float4(InVector,0.0), NiagaraSpriteVF.LocalToWorldInverseTransposed).xyz;
+}
 
 /*********************************************************************************
  *							Particle specific
@@ -448,11 +527,11 @@ void ComputeBillboardUVs(FVertexFactoryInput Input, out float2 UVForPosition, ou
 
 float3 GetNiagaraParticlePosition(uint InstanceID)
 {
-	return GetVec3(NiagaraSpriteVF.PositionDataOffset, InstanceID);
+	return TransformPosition(GetVec3(NiagaraSpriteVF.PositionDataOffset, InstanceID));
 }
 float3 GetNiagaraParticleVelocity(uint InstanceID)
 {
-	return GetVec3(NiagaraSpriteVF.VelocityDataOffset, InstanceID);
+	return TransformVector(GetVec3(NiagaraSpriteVF.VelocityDataOffset, InstanceID));
 }
 float2 GetNiagaraParticleSize(uint InstanceID)
 {
@@ -460,13 +539,53 @@ float2 GetNiagaraParticleSize(uint InstanceID)
 }
 float GetNiagaraParticleRotation(uint InstanceID)
 {
-	return GetFloat(NiagaraSpriteVF.RotationDataOffset, InstanceID);
+	return (GetFloat(NiagaraSpriteVF.RotationDataOffset, InstanceID) / 180.0f) * PI;
 }
 float4 GetNiagaraParticleColor(uint InstanceID)
 {
 	return GetVec4(NiagaraSpriteVF.ColorDataOffset, InstanceID);
 }
 
+float2 GetNiagaraUVScale(uint InstanceID)
+{
+	if(NiagaraSpriteVF.UVScaleOffset  == -1)
+	{
+		return float2(1.0f, 1.0f);
+	}
+	return GetVec2(NiagaraSpriteVF.UVScaleOffset, InstanceID);	
+}
+
+float GetNiagaraParticleRandom(uint InstanceID)
+{
+	if(NiagaraSpriteVF.ParticleRandomOffset  == -1)
+	{
+		return 0.0f;
+	}
+	return GetFloat(NiagaraSpriteVF.ParticleRandomOffset, InstanceID);	
+}
+
+
+
+float GetNiagaraCameraOffset(uint InstanceID)
+{
+	if (NiagaraSpriteVF.CameraOffsetOffset == -1)
+	{
+		return 0.0f;
+	}
+	return GetFloat(NiagaraSpriteVF.CameraOffsetOffset, InstanceID);
+}
+
+#if USE_DYNAMIC_PARAMETERS
+float4 GetNiagaraParticleDynamicParameters(uint InstanceID)
+{
+	if (NiagaraSpriteVF.MaterialParamDataOffset == -1)
+	{
+		return float4(1.0f, 1.0f, 1.0f, 1.0f);
+	}
+	return GetVec4(NiagaraSpriteVF.MaterialParamDataOffset, InstanceID);
+}
+#endif
+
 float GetNiagaraParticleSubimage(uint InstanceID)
 {
 	float Val = 0.0f;
@@ -484,7 +603,7 @@ float3 GetNiagaraParticleFacingVector(uint InstanceID)
 	{
 		Val = GetVec3(NiagaraSpriteVF.FacingOffset, InstanceID);
 	}
-	return Val;
+	return TransformVector(Val);
 }
 
 float3 GetNiagaraParticleAlignmentVector(uint InstanceID)
@@ -494,7 +613,7 @@ float3 GetNiagaraParticleAlignmentVector(uint InstanceID)
 	{
 		Val = GetVec3(NiagaraSpriteVF.AlignmentOffset, InstanceID);
 	}
-	return Val;
+	return TransformVector(Val);
 }
 
 
@@ -509,10 +628,25 @@ FVertexFactoryIntermediates GetVertexFactoryIntermediates(FVertexFactoryInput In
 	Input.Color	= GetNiagaraParticleColor(Input.InstanceId);
 	Input.CustomFacingVector = GetNiagaraParticleFacingVector(Input.InstanceId);
 	Input.CustomAlignmentVector = GetNiagaraParticleAlignmentVector(Input.InstanceId);
+#if USE_DYNAMIC_PARAMETERS
+	Input.DynamicParameter = GetNiagaraParticleDynamicParameters(Input.InstanceId);
+#endif
+
 
 	// World position.
+	
 	float4 ParticleWorldPosition = mul(float4(Input.Position.xyz,1), Primitive.LocalToWorld);
 	float4 ParticleOldWorldPosition = mul(float4(Input.OldPosition.xyz,1), Primitive.LocalToWorld); 
+	float CameraOffset = GetNiagaraCameraOffset(Input.InstanceId);
+	float3 CameraVector = ResolvedView.WorldCameraOrigin.xyz - ParticleWorldPosition.xyz;
+	float3 CameraDirection = SafeNormalize(CameraVector);
+
+	// Do we need the hack that forces particles at the camera origin to be *really* behind the camera origin?
+	float3 OffsetPosition = CameraDirection * CameraOffset;
+
+	ParticleWorldPosition.xyz = ParticleWorldPosition.xyz + OffsetPosition;
+	ParticleOldWorldPosition.xyz = ParticleOldWorldPosition.xyz + OffsetPosition;
+	
 	Intermediates.ParticleWorldPosition = ParticleWorldPosition.xyz;
 
 	// Tangents.
@@ -526,6 +660,9 @@ FVertexFactoryIntermediates GetVertexFactoryIntermediates(FVertexFactoryInput In
 	float2 UVForTexturingUnflipped;
 	ComputeBillboardUVs(Input, UVForPosition, UVForTexturing, UVForTexturingUnflipped);
 
+	float2 UVScale = GetNiagaraUVScale(Input.InstanceId);
+	UVForTexturing.xy = float2(0.5f, 0.5f) + (UVForTexturing.xy - float2(0.5f, 0.5f) ) *  UVScale.xy; 
+
 	// Vertex position.
 	float4 VertexWorldPosition = float4(ParticleWorldPosition.xyz + ResolvedView.PreViewTranslation,1);
 	float2 Size = abs(Input.SizeRotSubImage.xy);
@@ -548,7 +685,7 @@ FVertexFactoryIntermediates GetVertexFactoryIntermediates(FVertexFactoryInput In
 	Intermediates.TexCoord_Unflipped.xy = (float2(SubImageAH, SubImageAV) + UVForTexturingUnflipped) * NiagaraSpriteVF.SubImageSize.zw;
 	Intermediates.TexCoord_Unflipped.zw = (float2(SubImageBH, SubImageBV) + UVForTexturingUnflipped) * NiagaraSpriteVF.SubImageSize.zw;
 #endif
-	Intermediates.SubImageLerp = SubImageLerp;
+	Intermediates.SubImageLerp = (NiagaraSpriteVF.SubImageBlendMode == 1) ? SubImageLerp : 0.0f;
 
 	Intermediates.Color = Input.Color;
 
@@ -570,7 +707,7 @@ FVertexFactoryIntermediates GetVertexFactoryIntermediates(FVertexFactoryInput In
 
 #if USE_PARTICLE_VELOCITY
 	float3 ParticleVelocity = (ParticleWorldPosition.xyz - ParticleOldWorldPosition.xyz) * NiagaraSpriteVF.InvDeltaSeconds;
-	float3 DirectionVector = lerp(ParticleVelocity, Input.CustomAlignmentVector, UseCustomAlignment);
+	float3 DirectionVector = lerp(ParticleVelocity, Input.CustomAlignmentVector, ParticleAlignmentMode == ALIGNMENT_MODE_CUSTOMALIGNMENT);
 	Intermediates.ParticleVelocity.xyz = normalize(DirectionVector);
 	Intermediates.ParticleVelocity.w = length(DirectionVector);
 #endif
@@ -579,6 +716,10 @@ FVertexFactoryIntermediates GetVertexFactoryIntermediates(FVertexFactoryInput In
 	Intermediates.RelativeTime = Input.Position.w;
 #endif
 
+#if USE_PARTICLE_RANDOM
+	Intermediates.ParticleRandom = GetNiagaraParticleRandom(Input.InstanceId);
+#endif
+
 	Intermediates.TangentToLocal = CalcTangentBasis(Intermediates);
 
 #if USE_PARTICLE_LIGHTING_OFFSET
@@ -687,6 +828,11 @@ FVertexFactoryInterpolantsVSToPS VertexFactoryGetInterpolantsVSToPS(FVertexFacto
 	Interpolants.ParticleSize = Intermediates.ParticleSize;
 #endif
 
+#if USE_PARTICLE_RANDOM
+	Interpolants.ParticleRandom = Intermediates.ParticleRandom;
+#endif
+
+
 #if INSTANCED_STEREO
 	Interpolants.PackedEyeIndex = 0;
 #endif
diff --git a/Engine/Shaders/Private/ParticleGPUSpriteVertexFactory.ush b/Engine/Shaders/Private/ParticleGPUSpriteVertexFactory.ush
--- a/Engine/Shaders/Private/ParticleGPUSpriteVertexFactory.ush
+++ b/Engine/Shaders/Private/ParticleGPUSpriteVertexFactory.ush
@@ -272,7 +272,11 @@ FVertexFactoryIntermediates GetVertexFactoryIntermediates( FVertexFactoryInput I
 	float2 UVFlipOffset = float2( AttributeSample.x < 0.5 ? 0.0f : -0.5, AttributeSample.y < 0.5 ? 0.0f : -0.5);
 	float2 InitialSize = (AttributeSample.xy + UVFlipOffset) * float2(2.0,2.0);
 	float2 SizeScale = MiscCurve.xy * EmitterDynamicUniforms.LocalToWorldScale;
-	float2 SizeScaleBySpeed = clamp(EmitterUniforms.SizeBySpeed.xy * ParticleSpeed,float2(1,1), EmitterUniforms.SizeBySpeed.zw);
+
+	//Explicit ordering of clamp's min max to bypass compiler diff on PS4.
+	float2 SizeScaleBySpeed = max(EmitterUniforms.SizeBySpeed.xy * ParticleSpeed,float2(1,1));
+	SizeScaleBySpeed = min(SizeScaleBySpeed, EmitterUniforms.SizeBySpeed.zw);
+
 	float2 Size = InitialSize * SizeScale * SizeScaleBySpeed * IsAlive.xx;
 	
 	float2 FlippedUV = Input.TexCoord.xy;
diff --git a/Engine/Shaders/Private/ParticleSimulationShader.usf b/Engine/Shaders/Private/ParticleSimulationShader.usf
--- a/Engine/Shaders/Private/ParticleSimulationShader.usf
+++ b/Engine/Shaders/Private/ParticleSimulationShader.usf
@@ -106,13 +106,16 @@ Texture2D CurveTexture;
 SamplerState CurveTextureSampler;
 
 /** Textures from which to sample vector forces. */
-#if MAX_VECTOR_FIELDS != 4
+#if MAX_VECTOR_FIELDS != 4 && MAX_VECTOR_FIELDS != 1
 #error This must match MAX_VECTOR_FIELDS in C++ land
 #endif
 Texture3D VectorFieldTextures0;
+
+#if MAX_VECTOR_FIELDS > 1
 Texture3D VectorFieldTextures1;
 Texture3D VectorFieldTextures2;
 Texture3D VectorFieldTextures3;
+#endif
 //Texture3D VectorFieldTextures4;
 //Texture3D VectorFieldTextures5;
 //Texture3D VectorFieldTextures6;
@@ -121,9 +124,11 @@ Texture3D VectorFieldTextures3;
 //Texture3D VectorFieldTextures9;
 
 SamplerState VectorFieldTexturesSampler0;
+#if MAX_VECTOR_FIELDS > 1
 SamplerState VectorFieldTexturesSampler1;
 SamplerState VectorFieldTexturesSampler2;
 SamplerState VectorFieldTexturesSampler3;
+#endif
 //SamplerState VectorFieldTexturesSampler4;
 //SamplerState VectorFieldTexturesSampler5;
 //SamplerState VectorFieldTexturesSampler6;
@@ -161,10 +166,14 @@ float3 ComputeOrbitVelocity(float Time, float RandomOrbit)
  */
  float3 SampleVectorFieldTexture(int Index, float3 UV)
  {
+#if MAX_VECTOR_FIELDS == 1
+	 return Texture3DSample(VectorFieldTextures0, VectorFieldTexturesSampler0, UV).xyz;
+#else
 	if (Index == 0) return Texture3DSample(VectorFieldTextures0, VectorFieldTexturesSampler0, UV).xyz;
 	if (Index == 1) return Texture3DSample(VectorFieldTextures1, VectorFieldTexturesSampler1, UV).xyz;
 	if (Index == 2) return Texture3DSample(VectorFieldTextures2, VectorFieldTexturesSampler2, UV).xyz;
 	return Texture3DSample(VectorFieldTextures3, VectorFieldTexturesSampler3, UV).xyz;
+#endif
 
 //	if (Index == 3) return Texture3DSample(VectorFieldTextures3, VectorFieldTexturesSampler3, UV).xyz;
 //	if (Index == 4) return Texture3DSample(VectorFieldTextures4, VectorFieldTexturesSampler4, UV).xyz;
